#android Turn on DisjointTimerQueries on Adreno 330-based devices with a driver later than version 53.0. This is a workaround for this driver, and is a drastic improvement to GPU performance!

diff --git a/Engine/Config/BaseDeviceProfiles.ini b/Engine/Config/BaseDeviceProfiles.ini
--- a/Engine/Config/BaseDeviceProfiles.ini
+++ b/Engine/Config/BaseDeviceProfiles.ini
@@ -179,6 +179,11 @@ BaseProfileName=Android_Mid
 DeviceType=Android
 BaseProfileName=Android_High
 
+[Android_Adreno330_Ver53 DeviceProfile]
+DeviceType=Android
+BaseProfileName=Android_Adreno330
++CVars=r.DisjointTimerQueries=1
+
 [PS4 DeviceProfile]
 DeviceType=PS4
 BaseProfileName=
diff --git a/Engine/Plugins/Runtime/AndroidDeviceProfileSelector/Source/AndroidDeviceProfileSelector/Private/AndroidDeviceProfileSelectorModule.cpp b/Engine/Plugins/Runtime/AndroidDeviceProfileSelector/Source/AndroidDeviceProfileSelector/Private/AndroidDeviceProfileSelectorModule.cpp
--- a/Engine/Plugins/Runtime/AndroidDeviceProfileSelector/Source/AndroidDeviceProfileSelector/Private/AndroidDeviceProfileSelectorModule.cpp
+++ b/Engine/Plugins/Runtime/AndroidDeviceProfileSelector/Source/AndroidDeviceProfileSelector/Private/AndroidDeviceProfileSelectorModule.cpp
@@ -33,7 +33,21 @@ FString const FAndroidDeviceProfileSelectorModule::GetRuntimeDeviceProfileName()
 	}
 	else if (FCString::Stricmp(*GPUFamily, TEXT("Adreno (TM) 330")) == 0)
 	{
-		ProfileName = TEXT("Android_Adreno330");
+		FString GLVersion = FAndroidMisc::GetGLVersion();
+		const TCHAR* Version = FCString::Stristr(*GLVersion, TEXT("ES 3.0 V@"));
+		int32 VersionNumber = 0;
+		if ( Version )
+		{
+			VersionNumber = FCString::Atoi(Version + 9);
+		}
+		if ( VersionNumber >= 53 )
+		{
+			ProfileName = TEXT("Android_Adreno330_Ver53");
+		}
+		else
+		{
+			ProfileName = TEXT("Android_Adreno330");
+		}
 	}
 
 	UE_LOG(LogAndroid, Log, TEXT("Selected Device Profile: [%s]"), *ProfileName);
diff --git a/Engine/Source/Runtime/Core/Public/Android/AndroidMisc.h b/Engine/Source/Runtime/Core/Public/Android/AndroidMisc.h
--- a/Engine/Source/Runtime/Core/Public/Android/AndroidMisc.h
+++ b/Engine/Source/Runtime/Core/Public/Android/AndroidMisc.h
@@ -78,6 +78,7 @@ struct CORE_API FAndroidMisc : public FGenericPlatformMisc
 
 	// ANDROID ONLY:
 	static FString GetGPUFamily();
+	static FString GetGLVersion();
 	static bool SupportsFloatingPointRenderTargets();
 
 #if !UE_BUILD_SHIPPING
diff --git a/Engine/Source/Runtime/OpenGLDrv/Private/Android/AndroidOpenGL.cpp b/Engine/Source/Runtime/OpenGLDrv/Private/Android/AndroidOpenGL.cpp
--- a/Engine/Source/Runtime/OpenGLDrv/Private/Android/AndroidOpenGL.cpp
+++ b/Engine/Source/Runtime/OpenGLDrv/Private/Android/AndroidOpenGL.cpp
@@ -338,6 +338,7 @@ public:
 	}
 
 	FString GPUFamily;
+	FString GLVersion;
 	bool bSupportsFloatingPointRenderTargets;
 	TArray<FString> TargetPlatformNames;
 
@@ -362,6 +363,8 @@ private:
 		GPUFamily = (const ANSICHAR*)glGetString(GL_RENDERER);
 		check(!GPUFamily.IsEmpty());
 
+		GLVersion = (const ANSICHAR*)glGetString(GL_VERSION);
+
 		// highest priority is the per-texture version
 		if (FAndroidOpenGL::SupportsASTC())
 		{
@@ -399,6 +402,11 @@ FString FAndroidMisc::GetGPUFamily()
 	return FAndroidGPUInfo::Get().GPUFamily;
 }
 
+FString FAndroidMisc::GetGLVersion()
+{
+	return FAndroidGPUInfo::Get().GLVersion;
+}
+
 bool FAndroidMisc::SupportsFloatingPointRenderTargets()
 {
 	return FAndroidGPUInfo::Get().bSupportsFloatingPointRenderTargets;
diff --git a/Engine/Source/Runtime/OpenGLDrv/Private/OpenGLES2.cpp b/Engine/Source/Runtime/OpenGLDrv/Private/OpenGLES2.cpp
--- a/Engine/Source/Runtime/OpenGLDrv/Private/OpenGLES2.cpp
+++ b/Engine/Source/Runtime/OpenGLDrv/Private/OpenGLES2.cpp
@@ -26,15 +26,11 @@ bool FOpenGLES2::bSupportsOcclusionQueries = false;
 /** GL_EXT_disjoint_timer_query */
 bool FOpenGLES2::bSupportsDisjointTimeQueries = false;
 
-
-#if !UE_BUILD_SHIPPING
 static TAutoConsoleVariable<int32> CVarDisjointTimerQueries(
 	TEXT("r.DisjointTimerQueries"),
 	0,
 	TEXT("If set to 1, allows GPU time to be measured (e.g. STAT UNIT). It defaults to 0 because some devices supports it but very slowly."),
 	ECVF_RenderThreadSafe);
-#endif
-
 
 /** GL_OES_rgb8_rgba8 */
 bool FOpenGLES2::bSupportsRGBA8 = false;
@@ -112,9 +108,7 @@ bool FOpenGLES2::bSupportsTextureStorageEXT = false;
 bool FOpenGLES2::SupportsDisjointTimeQueries()
 {
 	bool bAllowDisjointTimerQueries = false;
-#if !UE_BUILD_SHIPPING
 	bAllowDisjointTimerQueries = (CVarDisjointTimerQueries.GetValueOnRenderThread() == 1);
-#endif
 	return bSupportsDisjointTimeQueries && bAllowDisjointTimerQueries;
 }
 
