Added detection for Adreno 2xx GPUs and PowerVR 540 to Android profile selector
#android
#ue4

diff --git a/Engine/Config/BaseDeviceProfiles.ini b/Engine/Config/BaseDeviceProfiles.ini
--- a/Engine/Config/BaseDeviceProfiles.ini
+++ b/Engine/Config/BaseDeviceProfiles.ini
@@ -178,6 +178,26 @@ BaseProfileName=Android_Mid
 ;This offset needs to be set for the mosaic fallback to work on Galaxy S4 (SAMSUNG-IGH-I337)
 ;+CVars=r.DemosaicVposOffset=0.5
 
+[Android_Adreno200 DeviceProfile]
+DeviceType=Android
+BaseProfileName=Android
+
+[Android_Adreno203 DeviceProfile]
+DeviceType=Android
+BaseProfileName=Android
+
+[Android_Adreno205 DeviceProfile]
+DeviceType=Android
+BaseProfileName=Android
+
+[Android_Adreno220 DeviceProfile]
+DeviceType=Android
+BaseProfileName=Android
+
+[Android_Adreno225 DeviceProfile]
+DeviceType=Android
+BaseProfileName=Android
+
 [Android_Adreno330 DeviceProfile]
 DeviceType=Android
 BaseProfileName=Android_High
@@ -187,6 +207,14 @@ DeviceType=Android
 BaseProfileName=Android_Adreno330
 +CVars=r.DisjointTimerQueries=1
 
+[Android_PowerVR540 DeviceProfile]
+DeviceType=Android
+BaseProfileName=Android
+
+[Android_PowerVR544 DeviceProfile]
+DeviceType=Android
+BaseProfileName=Android
+
 [PS4 DeviceProfile]
 DeviceType=PS4
 BaseProfileName=
diff --git a/Engine/Plugins/Runtime/AndroidDeviceProfileSelector/Source/AndroidDeviceProfileSelector/Private/AndroidDeviceProfileSelectorModule.cpp b/Engine/Plugins/Runtime/AndroidDeviceProfileSelector/Source/AndroidDeviceProfileSelector/Private/AndroidDeviceProfileSelectorModule.cpp
--- a/Engine/Plugins/Runtime/AndroidDeviceProfileSelector/Source/AndroidDeviceProfileSelector/Private/AndroidDeviceProfileSelectorModule.cpp
+++ b/Engine/Plugins/Runtime/AndroidDeviceProfileSelector/Source/AndroidDeviceProfileSelector/Private/AndroidDeviceProfileSelectorModule.cpp
@@ -27,7 +27,31 @@ FString const FAndroidDeviceProfileSelectorModule::GetRuntimeDeviceProfileName()
 
 	UE_LOG(LogAndroid, Log,TEXT("Default profile:%s GPUFamily:%s"),*ProfileName,*GPUFamily);
 
-	if (FCString::Stricmp(*GPUFamily, TEXT("Adreno (TM) 320")) == 0)
+	if (FCString::Stricmp(*GPUFamily, TEXT("PowerVR SGX 540")) == 0)
+	{
+		ProfileName = TEXT("Android_PowerVR540");
+	}
+	else if (FCString::Stricmp(*GPUFamily, TEXT("Adreno (TM) 200")) == 0)
+	{
+		ProfileName = TEXT("Android_Adreno200");
+	}
+	else if (FCString::Stricmp(*GPUFamily, TEXT("Adreno (TM) 203")) == 0)
+	{
+		ProfileName = TEXT("Android_Adreno203");
+	}
+	else if (FCString::Stricmp(*GPUFamily, TEXT("Adreno (TM) 205")) == 0)
+	{
+		ProfileName = TEXT("Android_Adreno205");
+	}
+	else if (FCString::Stricmp(*GPUFamily, TEXT("Adreno (TM) 220")) == 0)
+	{
+		ProfileName = TEXT("Android_Adreno220");
+	}
+	else if (FCString::Stricmp(*GPUFamily, TEXT("Adreno (TM) 225")) == 0)
+	{
+		ProfileName = TEXT("Android_Adreno225");
+	}
+	else if (FCString::Stricmp(*GPUFamily, TEXT("Adreno (TM) 320")) == 0)
 	{
 		ProfileName = TEXT("Android_Adreno320");
 	}
