Added device profiles for Mali and improved for unknown devices
UEMOB-12

diff --git a/Engine/Config/BaseDeviceProfiles.ini b/Engine/Config/BaseDeviceProfiles.ini
--- a/Engine/Config/BaseDeviceProfiles.ini
+++ b/Engine/Config/BaseDeviceProfiles.ini
@@ -194,7 +194,7 @@ BaseProfileName=Android
 +CVars=r.LightShaftQuality=1
 +CVars=r.MobileContentScaleFactor=1.0
 
-[Android_Unrecognized DeviceProfile]
+[Android_Default DeviceProfile]
 DeviceType=Android
 BaseProfileName=Android_Mid
 
@@ -251,6 +251,18 @@ DeviceType=Android
 BaseProfileName=Android
 +CVars=r.BloomQuality=1
 
+[Android_Mali_4xx DeviceProfile]
+DeviceType=Android
+BaseProfileName=Android_Low
+
+[Android_Mali_T6xx DeviceProfile]
+DeviceType=Android
+BaseProfileName=Android_Mid
+
+[Android_Mali_T7xx DeviceProfile]
+DeviceType=Android
+BaseProfileName=Android_High
+
 [PS4 DeviceProfile]
 DeviceType=PS4
 BaseProfileName=
diff --git a/Engine/Plugins/Runtime/AndroidDeviceProfileSelector/Source/AndroidDeviceProfileSelector/Private/AndroidDeviceProfileSelectorModule.cpp b/Engine/Plugins/Runtime/AndroidDeviceProfileSelector/Source/AndroidDeviceProfileSelector/Private/AndroidDeviceProfileSelectorModule.cpp
--- a/Engine/Plugins/Runtime/AndroidDeviceProfileSelector/Source/AndroidDeviceProfileSelector/Private/AndroidDeviceProfileSelectorModule.cpp
+++ b/Engine/Plugins/Runtime/AndroidDeviceProfileSelector/Source/AndroidDeviceProfileSelector/Private/AndroidDeviceProfileSelectorModule.cpp
@@ -103,8 +103,32 @@ FString const FAndroidDeviceProfileSelectorModule::GetRuntimeDeviceProfileName()
 			ProfileName = TEXT("Android_IntelHD");
 		}
 	}
+	else if (GPUFamily.StartsWith(TEXT("Mali-4")))
+	{
+		ProfileName = TEXT("Android_Mali_4xx");
+	}
+	else if (GPUFamily.StartsWith(TEXT("Mali-T6")))
+	{
+		ProfileName = TEXT("Android_Mali_T6xx");
+	}
+	else if (GPUFamily.StartsWith(TEXT("Mali-T7")))
+	{
+		ProfileName = TEXT("Android_Mali_T7xx");
+	}
+	else
+	{
+		FString GLVersion = FAndroidMisc::GetGLVersion();
 
-
+		if (GLVersion.StartsWith(TEXT("OpenGL ES 3.")))
+		{
+			ProfileName = TEXT("Android_Mid");
+		}
+		else if (GLVersion.StartsWith(TEXT("OpenGL ES 2.")))
+		{
+			ProfileName = TEXT("Android_Low");
+		}
+	}
+	
 	UE_LOG(LogAndroid, Log, TEXT("Selected Device Profile: [%s]"), *ProfileName);
 
 	// If no type was obtained, just select IOS as default
diff --git a/Engine/Source/Runtime/Core/Public/Android/AndroidMisc.h b/Engine/Source/Runtime/Core/Public/Android/AndroidMisc.h
--- a/Engine/Source/Runtime/Core/Public/Android/AndroidMisc.h
+++ b/Engine/Source/Runtime/Core/Public/Android/AndroidMisc.h
@@ -42,7 +42,7 @@ struct CORE_API FAndroidMisc : public FGenericPlatformMisc
 	static bool GetUseVirtualJoysticks();
 	static uint32 GetCharKeyMap(uint16* KeyCodes, FString* KeyNames, uint32 MaxMappings);
 	static uint32 GetKeyMap( uint16* KeyCodes, FString* KeyNames, uint32 MaxMappings );
-	static const TCHAR* GetDefaultDeviceProfileName() { return TEXT("Android"); }
+	static const TCHAR* GetDefaultDeviceProfileName() { return TEXT("Android_Default"); }
 	static bool GetVolumeButtonsHandledBySystem();
 	static void SetVolumeButtonsHandledBySystem(bool enabled);
 
