User-extendible Android device profile selection mechanism (UEMOB-37)

diff --git a/Engine/Config/BaseDeviceProfiles.ini b/Engine/Config/BaseDeviceProfiles.ini
--- a/Engine/Config/BaseDeviceProfiles.ini
+++ b/Engine/Config/BaseDeviceProfiles.ini
@@ -167,6 +167,27 @@ BaseProfileName=IOS
 +CVars=r.DepthOfFieldQuality=1
 +CVars=r.LightShaftQuality=1
 
+[/Script/AndroidDeviceProfileSelector.AndroidDeviceProfileMatchingRules]
+MatchProfile=(Profile="Android_PowerVR540",Match=((SourceType=SRC_GpuFamily,CompareType=CMP_Equal,MatchString="PowerVR SGX 540")))
++MatchProfile=(Profile="Android_Adreno200",Match=((SourceType=SRC_GpuFamily,CompareType=CMP_Equal,MatchString="Adreno (TM) 200")))
++MatchProfile=(Profile="Android_Adreno203",Match=((SourceType=SRC_GpuFamily,CompareType=CMP_Equal,MatchString="Adreno (TM) 203")))
++MatchProfile=(Profile="Android_Adreno205",Match=((SourceType=SRC_GpuFamily,CompareType=CMP_Equal,MatchString="Adreno (TM) 205")))
++MatchProfile=(Profile="Android_Adreno220",Match=((SourceType=SRC_GpuFamily,CompareType=CMP_Equal,MatchString="Adreno (TM) 220")))
++MatchProfile=(Profile="Android_Adreno225",Match=((SourceType=SRC_GpuFamily,CompareType=CMP_Equal,MatchString="Adreno (TM) 225")))
++MatchProfile=(Profile="Android_Adreno320",Match=((SourceType=SRC_GpuFamily,CompareType=CMP_Equal,MatchString="Adreno (TM) 320")))
++MatchProfile=(Profile="Android_Adreno330_Ver53",Match=((SourceType=SRC_GpuFamily,CompareType=CMP_Equal,MatchString="Adreno (TM) 330"),(SourceType=SRC_GlVersion,CompareType=CMP_Regex,MatchString="ES\\ 3\\.0\\ V@([0-9]+)"),(SourceType=SRC_PreviousRegexMatch,CompareType=CMP_GreaterEqual,MatchString="53")))
++MatchProfile=(Profile="Android_Adreno330",Match=((SourceType=SRC_GpuFamily,CompareType=CMP_Equal,MatchString="Adreno (TM) 330")))
++MatchProfile=(Profile="Android_TegraK1_GL4",Match=((SourceType=SRC_GpuFamily,CompareType=CMP_Equal,MatchString="NVIDIA Tegra"),(SourceType=SRC_GlVersion,CompareType=CMP_Regex,MatchString="^4\\.")))
++MatchProfile=(Profile="Android_TegraK1",Match=((SourceType=SRC_GpuFamily,CompareType=CMP_Equal,MatchString="NVIDIA Tegra"),(SourceType=SRC_GlVersion,CompareType=CMP_Regex,MatchString="^OpenGL\\ ES 3\\.")))
++MatchProfile=(Profile="Android_Tegra4",Match=((SourceType=SRC_GpuFamily,CompareType=CMP_Equal,MatchString="NVIDIA Tegra"),(SourceType=SRC_GlVersion,CompareType=CMP_Regex,MatchString="^OpenGL\\ ES 2\\.")))
++MatchProfile=(Profile="Android_IntelHD_ES3",Match=((SourceType=SRC_GpuFamily,CompareType=CMP_Equal,MatchString="Intel(R) HD Graphics"),(SourceType=SRC_GlVersion,CompareType=CMP_Regex,MatchString="^OpenGL\\ ES 3\\.")))
++MatchProfile=(Profile="Android_IntelHD",Match=((SourceType=SRC_GpuFamily,CompareType=CMP_Equal,MatchString="Intel(R) HD Graphics"),(SourceType=SRC_GlVersion,CompareType=CMP_Regex,MatchString="^OpenGL\\ ES 2\\.")))
++MatchProfile=(Profile="Android_Mali_4xx",Match=((SourceType=SRC_GpuFamily,CompareType=CMP_Regex,MatchString="^Mali\\-4")))
++MatchProfile=(Profile="Android_Mali_T6xx",Match=((SourceType=SRC_GpuFamily,CompareType=CMP_Regex,MatchString="^Mali\\-T6")))
++MatchProfile=(Profile="Android_Mali_T7xx",Match=((SourceType=SRC_GpuFamily,CompareType=CMP_Regex,MatchString="^Mali\\-T7")))
++MatchProfile=(Profile="Android_Mid",Match=((SourceType=SRC_GlVersion,CompareType=CMP_Regex,MatchString="^OpenGL\\ ES 3\\.")))
++MatchProfile=(Profile="Android_Low",Match=((SourceType=SRC_GlVersion,CompareType=CMP_Regex,MatchString="^OpenGL\\ ES 2\\.")))
+
 [Android DeviceProfile]
 DeviceType=Android
 BaseProfileName=
diff --git a/Engine/Plugins/Runtime/AndroidDeviceProfileSelector/Source/AndroidDeviceProfileSelector/AndroidDeviceProfileSelector.Build.cs b/Engine/Plugins/Runtime/AndroidDeviceProfileSelector/Source/AndroidDeviceProfileSelector/AndroidDeviceProfileSelector.Build.cs
--- a/Engine/Plugins/Runtime/AndroidDeviceProfileSelector/Source/AndroidDeviceProfileSelector/AndroidDeviceProfileSelector.Build.cs
+++ b/Engine/Plugins/Runtime/AndroidDeviceProfileSelector/Source/AndroidDeviceProfileSelector/AndroidDeviceProfileSelector.Build.cs
@@ -22,6 +22,7 @@ namespace UnrealBuildTool.Rules
 				new string[]
 				{
 					"Core",
+				    "CoreUObject",
 				}
 				);
 
diff --git a/Engine/Plugins/Runtime/AndroidDeviceProfileSelector/Source/AndroidDeviceProfileSelector/Private/AndroidDeviceProfileSelectorModule.cpp b/Engine/Plugins/Runtime/AndroidDeviceProfileSelector/Source/AndroidDeviceProfileSelector/Private/AndroidDeviceProfileSelectorModule.cpp
--- a/Engine/Plugins/Runtime/AndroidDeviceProfileSelector/Source/AndroidDeviceProfileSelector/Private/AndroidDeviceProfileSelectorModule.cpp
+++ b/Engine/Plugins/Runtime/AndroidDeviceProfileSelector/Source/AndroidDeviceProfileSelector/Private/AndroidDeviceProfileSelectorModule.cpp
@@ -1,136 +1,160 @@
 // Copyright 1998-2015 Epic Games, Inc. All Rights Reserved.
 
 #include "AndroidDeviceProfileSelectorPrivatePCH.h"
+#include "Regex.h"
 
 IMPLEMENT_MODULE(FAndroidDeviceProfileSelectorModule, AndroidDeviceProfileSelector);
 
+UAndroidDeviceProfileMatchingRules::UAndroidDeviceProfileMatchingRules(const FObjectInitializer& ObjectInitializer)
+: Super(ObjectInitializer)
+{
+}
 
 void FAndroidDeviceProfileSelectorModule::StartupModule()
 {
 }
 
-
 void FAndroidDeviceProfileSelectorModule::ShutdownModule()
 {
 }
 
-
 FString const FAndroidDeviceProfileSelectorModule::GetRuntimeDeviceProfileName()
 {
-	FString ProfileName = FPlatformMisc::GetDefaultDeviceProfileName();
-	if( ProfileName.IsEmpty() )
-	{
-		ProfileName = FPlatformProperties::PlatformName();
-	}
-
-	FString GPUFamily = FAndroidMisc::GetGPUFamily();
-
-	UE_LOG(LogAndroid, Log,TEXT("Default profile:%s GPUFamily:%s"),*ProfileName,*GPUFamily);
-
-	if (FCString::Stricmp(*GPUFamily, TEXT("PowerVR SGX 540")) == 0)
-	{
-		ProfileName = TEXT("Android_PowerVR540");
-	}
-	else if (FCString::Stricmp(*GPUFamily, TEXT("Adreno (TM) 200")) == 0)
-	{
-		ProfileName = TEXT("Android_Adreno200");
-	}
-	else if (FCString::Stricmp(*GPUFamily, TEXT("Adreno (TM) 203")) == 0)
-	{
-		ProfileName = TEXT("Android_Adreno203");
-	}
-	else if (FCString::Stricmp(*GPUFamily, TEXT("Adreno (TM) 205")) == 0)
-	{
-		ProfileName = TEXT("Android_Adreno205");
-	}
-	else if (FCString::Stricmp(*GPUFamily, TEXT("Adreno (TM) 220")) == 0)
-	{
-		ProfileName = TEXT("Android_Adreno220");
-	}
-	else if (FCString::Stricmp(*GPUFamily, TEXT("Adreno (TM) 225")) == 0)
-	{
-		ProfileName = TEXT("Android_Adreno225");
-	}
-	else if (FCString::Stricmp(*GPUFamily, TEXT("Adreno (TM) 320")) == 0)
-	{
-		ProfileName = TEXT("Android_Adreno320");
-	}
-	else if (FCString::Stricmp(*GPUFamily, TEXT("Adreno (TM) 330")) == 0)
+	static FString ProfileName; 
+	
+	if (ProfileName.IsEmpty())
 	{
-		FString GLVersion = FAndroidMisc::GetGLVersion();
-		const TCHAR* Version = FCString::Stristr(*GLVersion, TEXT("ES 3.0 V@"));
-		int32 VersionNumber = 0;
-		if ( Version )
-		{
-			VersionNumber = FCString::Atoi(Version + 9);
-		}
-		if ( VersionNumber >= 53 )
+		// Fallback profiles in case we do not match any rules
+		ProfileName = FPlatformMisc::GetDefaultDeviceProfileName();
+		if (ProfileName.IsEmpty())
 		{
-			ProfileName = TEXT("Android_Adreno330_Ver53");
+			ProfileName = FPlatformProperties::PlatformName();
 		}
-		else
-		{
-			ProfileName = TEXT("Android_Adreno330");
-		}
-	}
-	else if (FCString::Stricmp(*GPUFamily, TEXT("NVIDIA Tegra")) == 0)
-	{
-		FString GLVersion = FAndroidMisc::GetGLVersion();
 
-		if(GLVersion.StartsWith(TEXT("4.")))
-		{
-			ProfileName = TEXT("Android_TegraK1_GL4");
-		}
-		else if (GLVersion.StartsWith(TEXT("OpenGL ES 3.")))
-		{
-			ProfileName = TEXT("Android_TegraK1");
-		}
-		else if (GLVersion.StartsWith(TEXT("OpenGL ES 2.")))
-		{
-			ProfileName = TEXT("Android_Tegra4");
-		}
-	}
-	else if (GPUFamily.StartsWith(TEXT("Intel(R) HD Graphics")))
-	{
+		FString GPUFamily = FAndroidMisc::GetGPUFamily();
 		FString GLVersion = FAndroidMisc::GetGLVersion();
+		FString AndroidVersion = FAndroidMisc::GetAndroidVersion();
+		FString DeviceMake = FAndroidMisc::GetDeviceMake();
+		FString DeviceModel = FAndroidMisc::GetDeviceModel();
+		
+		// We need to initialize the class early as device profiles need to be evaluated before ProcessNewlyLoadedUObjects can be called.
+		extern UClass* Z_Construct_UClass_UAndroidDeviceProfileMatchingRules();
+		Z_Construct_UClass_UAndroidDeviceProfileMatchingRules();
 
-		if (GLVersion.StartsWith(TEXT("OpenGL ES 3.")))
-		{
-			ProfileName = TEXT("Android_IntelHD_ES3");
-		}
-		else if (GLVersion.StartsWith(TEXT("OpenGL ES 2.")))
-		{
-			ProfileName = TEXT("Android_IntelHD");
-		}
-	}
-	else if (GPUFamily.StartsWith(TEXT("Mali-4")))
-	{
-		ProfileName = TEXT("Android_Mali_4xx");
-	}
-	else if (GPUFamily.StartsWith(TEXT("Mali-T6")))
-	{
-		ProfileName = TEXT("Android_Mali_T6xx");
-	}
-	else if (GPUFamily.StartsWith(TEXT("Mali-T7")))
-	{
-		ProfileName = TEXT("Android_Mali_T7xx");
-	}
-	else
-	{
-		FString GLVersion = FAndroidMisc::GetGLVersion();
+		// Get the default object which will has the values from DeviceProfiles.ini
+		UAndroidDeviceProfileMatchingRules* Rules = Cast<UAndroidDeviceProfileMatchingRules>(UAndroidDeviceProfileMatchingRules::StaticClass()->GetDefaultObject());
+		check(Rules);
 
-		if (GLVersion.StartsWith(TEXT("OpenGL ES 3.")))
-		{
-			ProfileName = TEXT("Android_Mid");
-		}
-		else if (GLVersion.StartsWith(TEXT("OpenGL ES 2.")))
+		UE_LOG(LogAndroid, Log, TEXT("Checking %d rules from DeviceProfile ini file."), Rules->MatchProfile.Num());
+		UE_LOG(LogAndroid, Log, TEXT("  Default profile: %s"), *ProfileName);
+		UE_LOG(LogAndroid, Log, TEXT("  GpuFamily: %s"), *GPUFamily);
+		UE_LOG(LogAndroid, Log, TEXT("  GlVersion: %s"), *GLVersion);
+		UE_LOG(LogAndroid, Log, TEXT("  AndroidVersion: %s"), *AndroidVersion);
+		UE_LOG(LogAndroid, Log, TEXT("  DeviceMake: %s"), *DeviceMake);
+		UE_LOG(LogAndroid, Log, TEXT("  DeviceModel: %s"), *DeviceModel);
+
+		for (const FProfileMatch& Profile : Rules->MatchProfile)
 		{
-			ProfileName = TEXT("Android_Low");
+			FString PreviousRegexMatch;
+			bool bFoundMatch = true;
+			for (const FProfileMatchItem& Item : Profile.Match)
+			{
+				FString* SourceString = nullptr;
+				switch (Item.SourceType)
+				{
+				case SRC_PreviousRegexMatch:
+					SourceString = &PreviousRegexMatch;
+					break;
+				case SRC_GpuFamily:
+					SourceString = &GPUFamily;
+					break;
+				case SRC_GlVersion:
+					SourceString = &GLVersion;
+					break;
+				case SRC_AndroidVersion:
+					SourceString = &AndroidVersion;
+					break;
+				case SRC_DeviceMake:
+					SourceString = &DeviceMake;
+					break;
+				case SRC_DeviceModel:
+					SourceString = &DeviceModel;
+					break;
+				default:
+					continue;
+				}
+
+				switch (Item.CompareType)
+				{
+				case CMP_Equal:
+					if (*SourceString != Item.MatchString)
+					{
+						bFoundMatch = false;
+					}
+					break;
+				case CMP_Less:
+					if (FPlatformString::Atoi(**SourceString) >= FPlatformString::Atoi(*Item.MatchString))
+					{
+						bFoundMatch = false;
+					}
+					break;
+				case CMP_LessEqual:
+					if (FPlatformString::Atoi(**SourceString) > FPlatformString::Atoi(*Item.MatchString))
+					{
+						bFoundMatch = false;
+					}
+					break;
+				case CMP_Greater:
+					if (FPlatformString::Atoi(**SourceString) <= FPlatformString::Atoi(*Item.MatchString))
+					{
+						bFoundMatch = false;
+					}
+					break;
+				case CMP_GreaterEqual:
+					if (FPlatformString::Atoi(**SourceString) < FPlatformString::Atoi(*Item.MatchString))
+					{
+						bFoundMatch = false;
+					}
+					break;
+				case CMP_NotEqual:
+					if (*SourceString == Item.MatchString)
+					{
+						bFoundMatch = false;
+					}
+					break;
+				case CMP_Regex:
+					{
+						const FRegexPattern RegexPattern(Item.MatchString);
+						FRegexMatcher RegexMatcher(RegexPattern, *SourceString);
+						if (RegexMatcher.FindNext())
+						{
+							PreviousRegexMatch = RegexMatcher.GetCaptureGroup(1);
+						}
+						else
+						{
+							bFoundMatch = false;
+						}
+					}
+					break;
+				default:
+					bFoundMatch = false;
+				}
+
+				if (!bFoundMatch)
+				{
+					break;
+				}
+			}
+
+			if (bFoundMatch)
+			{
+				ProfileName = Profile.Profile;
+				break;
+			}
 		}
+
+		UE_LOG(LogAndroid, Log, TEXT("Selected Device Profile: [%s]"), *ProfileName);
 	}
-	
-	UE_LOG(LogAndroid, Log, TEXT("Selected Device Profile: [%s]"), *ProfileName);
 
-	// If no type was obtained, just select IOS as default
 	return ProfileName;
 }
\ No newline at end of file
diff --git a/Engine/Plugins/Runtime/AndroidDeviceProfileSelector/Source/AndroidDeviceProfileSelector/Public/AndroidDeviceProfileMatchingRules.h b/Engine/Plugins/Runtime/AndroidDeviceProfileSelector/Source/AndroidDeviceProfileSelector/Public/AndroidDeviceProfileMatchingRules.h
new file mode 100644
--- /dev/null
+++ b/Engine/Plugins/Runtime/AndroidDeviceProfileSelector/Source/AndroidDeviceProfileSelector/Public/AndroidDeviceProfileMatchingRules.h
@@ -0,0 +1,69 @@
+// Copyright 1998-2015 Epic Games, Inc. All Rights Reserved.
+
+#pragma once
+
+#include "AndroidDeviceProfileMatchingRules.generated.h"
+
+UENUM()
+enum ESourceType
+{
+	SRC_PreviousRegexMatch,
+	SRC_GpuFamily,
+	SRC_GlVersion,
+	SRC_AndroidVersion,
+	SRC_DeviceMake,
+	SRC_DeviceModel,
+	SRC_MAX,
+};
+
+UENUM()
+enum ECompareType
+{
+	CMP_Equal,
+	CMP_Less,
+	CMP_LessEqual,
+	CMP_Greater,
+	CMP_GreaterEqual,
+	CMP_NotEqual,
+	CMP_Regex,
+	CMP_MAX,
+};
+
+USTRUCT()
+struct FProfileMatchItem
+{
+	GENERATED_USTRUCT_BODY()
+
+	UPROPERTY()
+	TEnumAsByte<ESourceType> SourceType;
+
+	UPROPERTY()
+	TEnumAsByte<ECompareType> CompareType;
+
+	UPROPERTY()
+	FString MatchString;
+};
+
+USTRUCT()
+struct FProfileMatch
+{
+	GENERATED_USTRUCT_BODY()
+
+	UPROPERTY()
+	FString Profile;
+	
+	UPROPERTY()
+	TArray<FProfileMatchItem> Match;
+};
+
+UCLASS(config = DeviceProfiles)
+class UAndroidDeviceProfileMatchingRules : public UObject
+{
+	GENERATED_UCLASS_BODY()
+
+public:
+
+	/** Array of rules to match */
+	UPROPERTY(EditAnywhere, config, Category = "Matching Rules")
+	TArray<FProfileMatch> MatchProfile;
+};
diff --git a/Engine/Plugins/Runtime/AndroidDeviceProfileSelector/Source/AndroidDeviceProfileSelector/Public/AndroidDeviceProfileSelector.h b/Engine/Plugins/Runtime/AndroidDeviceProfileSelector/Source/AndroidDeviceProfileSelector/Public/AndroidDeviceProfileSelector.h
--- a/Engine/Plugins/Runtime/AndroidDeviceProfileSelector/Source/AndroidDeviceProfileSelector/Public/AndroidDeviceProfileSelector.h
+++ b/Engine/Plugins/Runtime/AndroidDeviceProfileSelector/Source/AndroidDeviceProfileSelector/Public/AndroidDeviceProfileSelector.h
@@ -6,3 +6,4 @@
 *****************************************************************************/
 #include "ModuleManager.h"
 #include "IDeviceProfileSelectorModule.h"
+#include "AndroidDeviceProfileMatchingRules.h"
