Allow hardware/chipset check in device profile matching rules for Android

#ue4
#android

diff --git a/Engine/Plugins/Runtime/AndroidDeviceProfileSelector/Source/AndroidDeviceProfileSelector/Private/AndroidDeviceProfileMatchingRules.h b/Engine/Plugins/Runtime/AndroidDeviceProfileSelector/Source/AndroidDeviceProfileSelector/Private/AndroidDeviceProfileMatchingRules.h
--- a/Engine/Plugins/Runtime/AndroidDeviceProfileSelector/Source/AndroidDeviceProfileSelector/Private/AndroidDeviceProfileMatchingRules.h
+++ b/Engine/Plugins/Runtime/AndroidDeviceProfileSelector/Source/AndroidDeviceProfileSelector/Private/AndroidDeviceProfileMatchingRules.h
@@ -20,6 +20,8 @@ enum ESourceType
 	SRC_UsingHoudini,
 	SRC_VulkanAvailable,
 	SRC_CommandLine,
+	SRC_Hardware,
+	SRC_Chipset,
 	SRC_MAX,
 };
 
diff --git a/Engine/Plugins/Runtime/AndroidDeviceProfileSelector/Source/AndroidDeviceProfileSelector/Private/AndroidDeviceProfileSelector.cpp b/Engine/Plugins/Runtime/AndroidDeviceProfileSelector/Source/AndroidDeviceProfileSelector/Private/AndroidDeviceProfileSelector.cpp
--- a/Engine/Plugins/Runtime/AndroidDeviceProfileSelector/Source/AndroidDeviceProfileSelector/Private/AndroidDeviceProfileSelector.cpp
+++ b/Engine/Plugins/Runtime/AndroidDeviceProfileSelector/Source/AndroidDeviceProfileSelector/Private/AndroidDeviceProfileSelector.cpp
@@ -30,7 +30,7 @@ static UAndroidDeviceProfileMatchingRules* GetAndroidDeviceProfileMatchingRules(
 	return Rules;
 }
 
-FString FAndroidDeviceProfileSelector::FindMatchingProfile(const FString& GPUFamily, const FString& GLVersion, const FString& AndroidVersion, const FString& DeviceMake, const FString& DeviceModel, const FString& DeviceBuildNumber, const FString& VulkanAvailable, const FString& VulkanVersion, const FString& UsingHoudini, const FString& ProfileName)
+FString FAndroidDeviceProfileSelector::FindMatchingProfile(const FString& GPUFamily, const FString& GLVersion, const FString& AndroidVersion, const FString& DeviceMake, const FString& DeviceModel, const FString& DeviceBuildNumber, const FString& VulkanAvailable, const FString& VulkanVersion, const FString& UsingHoudini, const FString& Hardware, const FString& Chipset, const FString& ProfileName)
 {
 	FString OutProfileName = ProfileName;
 	FString CommandLine = FCommandLine::Get();
@@ -77,6 +77,12 @@ FString FAndroidDeviceProfileSelector::FindMatchingProfile(const FString& GPUFam
 			case SRC_CommandLine:
 				SourceString = &CommandLine;
 				break;
+			case SRC_Hardware:
+				SourceString = &Hardware;
+				break;
+			case SRC_Chipset:
+				SourceString = &Chipset;
+				break;
 			default:
 				continue;
 			}
diff --git a/Engine/Plugins/Runtime/AndroidDeviceProfileSelector/Source/AndroidDeviceProfileSelector/Private/AndroidDeviceProfileSelectorModule.cpp b/Engine/Plugins/Runtime/AndroidDeviceProfileSelector/Source/AndroidDeviceProfileSelector/Private/AndroidDeviceProfileSelectorModule.cpp
--- a/Engine/Plugins/Runtime/AndroidDeviceProfileSelector/Source/AndroidDeviceProfileSelector/Private/AndroidDeviceProfileSelectorModule.cpp
+++ b/Engine/Plugins/Runtime/AndroidDeviceProfileSelector/Source/AndroidDeviceProfileSelector/Private/AndroidDeviceProfileSelectorModule.cpp
@@ -37,6 +37,8 @@ const FString FAndroidDeviceProfileSelectorModule::GetDeviceProfileName(const TM
 	FString DeviceModel = DeviceParameters.FindChecked("DeviceModel");
 	FString DeviceBuildNumber = DeviceParameters.FindChecked("DeviceBuildNumber");
 	FString UsingHoudini = DeviceParameters.FindChecked("UsingHoudini");
+	FString Hardware = DeviceParameters.FindChecked("Hardware");
+	FString Chipset = DeviceParameters.FindChecked("Chipset");
 
 	UE_LOG(LogAndroid, Log, TEXT("Checking %d rules from DeviceProfile ini file."), FAndroidDeviceProfileSelector::GetNumProfiles() );
 	UE_LOG(LogAndroid, Log, TEXT("  Default profile: %s"), *ProfileName);
@@ -49,8 +51,10 @@ const FString FAndroidDeviceProfileSelectorModule::GetDeviceProfileName(const TM
 	UE_LOG(LogAndroid, Log, TEXT("  DeviceModel: %s"), *DeviceModel);
 	UE_LOG(LogAndroid, Log, TEXT("  DeviceBuildNumber: %s"), *DeviceBuildNumber);
 	UE_LOG(LogAndroid, Log, TEXT("  UsingHoudini: %s"), *UsingHoudini);
+	UE_LOG(LogAndroid, Log, TEXT("  Hardware: %s"), *Hardware);
+	UE_LOG(LogAndroid, Log, TEXT("  Chipset: %s"), *Chipset);
 
-	ProfileName = FAndroidDeviceProfileSelector::FindMatchingProfile(GPUFamily, GLVersion, AndroidVersion, DeviceMake, DeviceModel, DeviceBuildNumber, VulkanAvailable, VulkanVersion, UsingHoudini, ProfileName);
+	ProfileName = FAndroidDeviceProfileSelector::FindMatchingProfile(GPUFamily, GLVersion, AndroidVersion, DeviceMake, DeviceModel, DeviceBuildNumber, VulkanAvailable, VulkanVersion, UsingHoudini, Hardware, Chipset, ProfileName);
 
 	UE_LOG(LogAndroid, Log, TEXT("Selected Device Profile: [%s]"), *ProfileName);
 
diff --git a/Engine/Plugins/Runtime/AndroidDeviceProfileSelector/Source/AndroidDeviceProfileSelector/Public/AndroidDeviceProfileSelector.h b/Engine/Plugins/Runtime/AndroidDeviceProfileSelector/Source/AndroidDeviceProfileSelector/Public/AndroidDeviceProfileSelector.h
--- a/Engine/Plugins/Runtime/AndroidDeviceProfileSelector/Source/AndroidDeviceProfileSelector/Public/AndroidDeviceProfileSelector.h
+++ b/Engine/Plugins/Runtime/AndroidDeviceProfileSelector/Source/AndroidDeviceProfileSelector/Public/AndroidDeviceProfileSelector.h
@@ -7,6 +7,6 @@
 class ANDROIDDEVICEPROFILESELECTOR_API FAndroidDeviceProfileSelector
 {
 public:
-	static FString FindMatchingProfile(const FString& GPUFamily, const FString& GLVersion, const FString& AndroidVersion, const FString& DeviceMake, const FString& DeviceModel, const FString& DeviceBuildNumber, const FString& VulkanAvailable, const FString& VulkanVersion, const FString& UsingHoudini, const FString& ProfileName);
+	static FString FindMatchingProfile(const FString& GPUFamily, const FString& GLVersion, const FString& AndroidVersion, const FString& DeviceMake, const FString& DeviceModel, const FString& DeviceBuildNumber, const FString& VulkanAvailable, const FString& VulkanVersion, const FString& UsingHoudini, const FString& Hardware, const FString& Chipset, const FString& ProfileName);
 	static int32 GetNumProfiles();
 };
diff --git a/Engine/Plugins/Runtime/AndroidDeviceProfileSelector/Source/AndroidDeviceProfileSelectorRuntime/Private/AndroidDeviceProfileSelectorRuntimeModule.cpp b/Engine/Plugins/Runtime/AndroidDeviceProfileSelector/Source/AndroidDeviceProfileSelectorRuntime/Private/AndroidDeviceProfileSelectorRuntimeModule.cpp
--- a/Engine/Plugins/Runtime/AndroidDeviceProfileSelector/Source/AndroidDeviceProfileSelectorRuntime/Private/AndroidDeviceProfileSelectorRuntimeModule.cpp
+++ b/Engine/Plugins/Runtime/AndroidDeviceProfileSelector/Source/AndroidDeviceProfileSelectorRuntime/Private/AndroidDeviceProfileSelectorRuntimeModule.cpp
@@ -57,6 +57,19 @@ FString const FAndroidDeviceProfileSelectorRuntimeModule::GetRuntimeDeviceProfil
 		FString DeviceMake = FAndroidMisc::GetDeviceMake();
 		FString DeviceModel = FAndroidMisc::GetDeviceModel();
 		FString DeviceBuildNumber = FAndroidMisc::GetDeviceBuildNumber();
+		
+		FString Hardware = FString(TEXT("unknown"));
+		FString Chipset = FString(TEXT("unknown"));
+		FString *HardwareLookup = FAndroidMisc::GetConfigRulesVariable(TEXT("hardware"));
+		if (HardwareLookup != nullptr)
+		{
+			Hardware = *HardwareLookup;
+		}
+		FString *ChipsetLookup = FAndroidMisc::GetConfigRulesVariable(TEXT("chipset"));
+		if (ChipsetLookup != nullptr)
+		{
+			Chipset = *ChipsetLookup;
+		}
 
 #if !(PLATFORM_ANDROID_X86 || PLATFORM_ANDROID_X64)
 		// Not running an Intel libUE4.so with Houdini library present means we're emulated
@@ -77,6 +90,8 @@ FString const FAndroidDeviceProfileSelectorRuntimeModule::GetRuntimeDeviceProfil
 		UE_LOG(LogAndroid, Log, TEXT("  DeviceModel: %s"), *DeviceModel);
 		UE_LOG(LogAndroid, Log, TEXT("  DeviceBuildNumber: %s"), *DeviceBuildNumber);
 		UE_LOG(LogAndroid, Log, TEXT("  UsingHoudini: %s"), *UsingHoudini);
+		UE_LOG(LogAndroid, Log, TEXT("  Hardware: %s"), *Hardware);
+		UE_LOG(LogAndroid, Log, TEXT("  Chipset: %s"), *Chipset);
 
 		CheckForJavaSurfaceViewWorkaround(DeviceMake, DeviceModel);
 
@@ -90,7 +105,7 @@ FString const FAndroidDeviceProfileSelectorRuntimeModule::GetRuntimeDeviceProfil
 		else
 		{
 			// Find a match with the DeviceProfiles matching rules
-			ProfileName = FAndroidDeviceProfileSelector::FindMatchingProfile(GPUFamily, GLVersion, AndroidVersion, DeviceMake, DeviceModel, DeviceBuildNumber, VulkanAvailable, VulkanVersion, UsingHoudini, ProfileName);
+			ProfileName = FAndroidDeviceProfileSelector::FindMatchingProfile(GPUFamily, GLVersion, AndroidVersion, DeviceMake, DeviceModel, DeviceBuildNumber, VulkanAvailable, VulkanVersion, UsingHoudini, Hardware, Chipset, ProfileName);
 			UE_LOG(LogAndroid, Log, TEXT("Selected Device Profile: [%s]"), *ProfileName);
 		}
 	}
diff --git a/Engine/Source/Editor/PIEPreviewDeviceProfileSelector/Private/PIEPreviewDevice.cpp b/Engine/Source/Editor/PIEPreviewDeviceProfileSelector/Private/PIEPreviewDevice.cpp
--- a/Engine/Source/Editor/PIEPreviewDeviceProfileSelector/Private/PIEPreviewDevice.cpp
+++ b/Engine/Source/Editor/PIEPreviewDeviceProfileSelector/Private/PIEPreviewDevice.cpp
@@ -454,6 +454,8 @@ FString FPIEPreviewDevice::GetProfile() const
 				DeviceParameters.Add("DeviceModel", AndroidProperties.DeviceModel);
 				DeviceParameters.Add("DeviceBuildNumber", AndroidProperties.DeviceBuildNumber);
 				DeviceParameters.Add("UsingHoudini", AndroidProperties.UsingHoudini ? "true" : "false");
+				DeviceParameters.Add("Hardware", AndroidProperties.Hardware);
+				DeviceParameters.Add("Chipset", AndroidProperties.Chipset);
 
 				FString PIEProfileName = AndroidDeviceProfileSelector->GetDeviceProfileName(DeviceParameters);
 				if (!PIEProfileName.IsEmpty())
diff --git a/Engine/Source/Editor/PIEPreviewDeviceSpecification/Public/PIEPreviewDeviceSpecification.h b/Engine/Source/Editor/PIEPreviewDeviceSpecification/Public/PIEPreviewDeviceSpecification.h
--- a/Engine/Source/Editor/PIEPreviewDeviceSpecification/Public/PIEPreviewDeviceSpecification.h
+++ b/Engine/Source/Editor/PIEPreviewDeviceSpecification/Public/PIEPreviewDeviceSpecification.h
@@ -42,6 +42,10 @@ public:
 	FString DeviceBuildNumber;
 	UPROPERTY()
 	bool UsingHoudini;
+	UPROPERTY()
+	FString Hardware;
+	UPROPERTY()
+	FString Chipset;
 };
 
 USTRUCT()
@@ -91,6 +95,10 @@ public:
 	bool VulkanAvailable;
 	UPROPERTY()
 	bool UsingHoudini;
+	UPROPERTY()
+	FString Hardware;
+	UPROPERTY()
+	FString Chipset;
 
 	UPROPERTY()
 	FPIERHIOverrideState GLES2RHIState;
