Added CMP_Hash operation to compare against a salted SHA1 hash when matching Android device profiles

diff --git a/Engine/Plugins/Runtime/AndroidDeviceProfileSelector/Source/AndroidDeviceProfileSelector/AndroidDeviceProfileSelector.Build.cs b/Engine/Plugins/Runtime/AndroidDeviceProfileSelector/Source/AndroidDeviceProfileSelector/AndroidDeviceProfileSelector.Build.cs
--- a/Engine/Plugins/Runtime/AndroidDeviceProfileSelector/Source/AndroidDeviceProfileSelector/AndroidDeviceProfileSelector.Build.cs
+++ b/Engine/Plugins/Runtime/AndroidDeviceProfileSelector/Source/AndroidDeviceProfileSelector/AndroidDeviceProfileSelector.Build.cs
@@ -1,4 +1,6 @@
 // Copyright 1998-2019 Epic Games, Inc. All Rights Reserved.
+using Tools.DotNETCommon;
+using System.IO;
 
 namespace UnrealBuildTool.Rules
 {
@@ -8,7 +10,17 @@ namespace UnrealBuildTool.Rules
 		{
 			ShortName = "AndroidDPS";
 
-			PublicDependencyModuleNames.AddRange(
+            string SecretFile = Path.Combine(ModuleDirectory, "Private", "NoRedist", "AndroidDeviceProfileSelectorSecrets.h");
+			if (File.Exists(SecretFile))
+            {
+                PrivateDefinitions.Add("ANDROIDDEVICEPROFILESELECTORSECRETS_H=1");
+            }
+            else
+            {
+                PrivateDefinitions.Add("ANDROIDDEVICEPROFILESELECTORSECRETS_H=0");
+            }
+
+            PublicDependencyModuleNames.AddRange(
 				new string[]
 				{
 					"Core",
diff --git a/Engine/Plugins/Runtime/AndroidDeviceProfileSelector/Source/AndroidDeviceProfileSelector/Private/AndroidDeviceProfileMatchingRules.h b/Engine/Plugins/Runtime/AndroidDeviceProfileSelector/Source/AndroidDeviceProfileSelector/Private/AndroidDeviceProfileMatchingRules.h
--- a/Engine/Plugins/Runtime/AndroidDeviceProfileSelector/Source/AndroidDeviceProfileSelector/Private/AndroidDeviceProfileMatchingRules.h
+++ b/Engine/Plugins/Runtime/AndroidDeviceProfileSelector/Source/AndroidDeviceProfileSelector/Private/AndroidDeviceProfileMatchingRules.h
@@ -41,6 +41,7 @@ enum ECompareType
 	CMP_GreaterIgnore,
 	CMP_GreaterEqualIgnore,
 	CMP_NotEqualIgnore,
+	CMP_Hash,
 	CMP_MAX,
 };
 
diff --git a/Engine/Plugins/Runtime/AndroidDeviceProfileSelector/Source/AndroidDeviceProfileSelector/Private/AndroidDeviceProfileSelector.cpp b/Engine/Plugins/Runtime/AndroidDeviceProfileSelector/Source/AndroidDeviceProfileSelector/Private/AndroidDeviceProfileSelector.cpp
--- a/Engine/Plugins/Runtime/AndroidDeviceProfileSelector/Source/AndroidDeviceProfileSelector/Private/AndroidDeviceProfileSelector.cpp
+++ b/Engine/Plugins/Runtime/AndroidDeviceProfileSelector/Source/AndroidDeviceProfileSelector/Private/AndroidDeviceProfileSelector.cpp
@@ -6,6 +6,12 @@
 #include "Templates/Casts.h"
 #include "Internationalization/Regex.h"
 #include "Misc/CommandLine.h"
+#include "Misc/SecureHash.h"
+#include "Containers/StringConv.h"
+
+#if ANDROIDDEVICEPROFILESELECTORSECRETS_H
+#include "NoRedist/AndroidDeviceProfileSelectorSecrets.h"
+#endif
 
 UAndroidDeviceProfileMatchingRules::UAndroidDeviceProfileMatchingRules(const FObjectInitializer& ObjectInitializer)
 	: Super(ObjectInitializer)
@@ -184,19 +190,46 @@ FString FAndroidDeviceProfileSelector::FindMatchingProfile(const FString& GPUFam
 				}
 				break;
 			case CMP_Regex:
-			{
-				const FRegexPattern RegexPattern(Item.MatchString);
-				FRegexMatcher RegexMatcher(RegexPattern, *SourceString);
-				if (RegexMatcher.FindNext())
-				{
-					PreviousRegexMatch = RegexMatcher.GetCaptureGroup(1);
-				}
-				else
 				{
-					bFoundMatch = false;
+					const FRegexPattern RegexPattern(Item.MatchString);
+					FRegexMatcher RegexMatcher(RegexPattern, *SourceString);
+					if (RegexMatcher.FindNext())
+					{
+						PreviousRegexMatch = RegexMatcher.GetCaptureGroup(1);
+					}
+					else
+					{
+						bFoundMatch = false;
+					}
 				}
-			}
 			break;
+				case CMP_Hash:
+				{
+					// Salt string is concatenated onto the end of the input text.
+					// For example the input string "PhoneModel" with salt "Salt" and pepper "Pepper" can be computed with
+					// % printf "PhoneModelSaltPepper" | openssl dgst -sha1 -hex
+					// resulting in d9e5cbd6b0e4dba00edd9de92cf64ee4c3f3a2db and would be stored in the matching rules as 
+					// "Salt|d9e5cbd6b0e4dba00edd9de92cf64ee4c3f3a2db". Salt is optional.
+					FString MatchHashString;
+					FString SaltString;
+					if (!Item.MatchString.Split(TEXT("|"), &SaltString, &MatchHashString))
+					{
+						MatchHashString = Item.MatchString;
+					}
+					FString HashInputString = *SourceString + SaltString
+#if ANDROIDDEVICEPROFILESELECTORSECRETS_H
+						+ HASH_PEPPER_SECRET_GUID.ToString()
+#endif
+						;
+
+					FSHAHash SourceHash;
+					FSHA1::HashBuffer(TCHAR_TO_ANSI(*HashInputString), HashInputString.Len(), SourceHash.Hash);
+					if (SourceHash.ToString() != MatchHashString)
+					{
+						bFoundMatch = false;
+					}
+				}
+				break;
 			default:
 				bFoundMatch = false;
 			}
