Take Android matching rules from AndroidDeviceProfles.ini

diff --git a/Engine/Plugins/Runtime/AndroidDeviceProfileSelector/Source/AndroidDeviceProfileSelector/Private/AndroidDeviceProfileMatchingRules.h b/Engine/Plugins/Runtime/AndroidDeviceProfileSelector/Source/AndroidDeviceProfileSelector/Private/AndroidDeviceProfileMatchingRules.h
--- a/Engine/Plugins/Runtime/AndroidDeviceProfileSelector/Source/AndroidDeviceProfileSelector/Private/AndroidDeviceProfileMatchingRules.h
+++ b/Engine/Plugins/Runtime/AndroidDeviceProfileSelector/Source/AndroidDeviceProfileSelector/Private/AndroidDeviceProfileMatchingRules.h
@@ -76,6 +76,12 @@ class UAndroidDeviceProfileMatchingRules : public UObject
 
 public:
 
+	// Use Android .ini files in the editor
+	virtual const TCHAR* GetConfigOverridePlatform() const override
+	{
+		return TEXT("Android");
+	}
+
 	/** Array of rules to match */
 	UPROPERTY(EditAnywhere, config, Category = "Matching Rules")
 	TArray<FProfileMatch> MatchProfile;
