Change AndroidDeviceProfileSelector editor module loading phase to manual

diff --git a/Engine/Plugins/Runtime/AndroidDeviceProfileSelector/AndroidDeviceProfileSelector.uplugin b/Engine/Plugins/Runtime/AndroidDeviceProfileSelector/AndroidDeviceProfileSelector.uplugin
--- a/Engine/Plugins/Runtime/AndroidDeviceProfileSelector/AndroidDeviceProfileSelector.uplugin
+++ b/Engine/Plugins/Runtime/AndroidDeviceProfileSelector/AndroidDeviceProfileSelector.uplugin
@@ -19,7 +19,7 @@
 		{
 			"Name" : "AndroidDeviceProfileSelector",
 			"Type": "Editor",
-			"LoadingPhase" : "PostConfigInit"
+			"LoadingPhase" : "None"
 		},
 		{
 			"Name" : "AndroidDeviceProfileSelectorRuntime",
