Made Android device rule profile hash checking case insensitive. Updated hashed device profiles.

diff --git a/Engine/Plugins/Runtime/AndroidDeviceProfileSelector/Source/AndroidDeviceProfileSelector/Private/AndroidDeviceProfileSelector.cpp b/Engine/Plugins/Runtime/AndroidDeviceProfileSelector/Source/AndroidDeviceProfileSelector/Private/AndroidDeviceProfileSelector.cpp
--- a/Engine/Plugins/Runtime/AndroidDeviceProfileSelector/Source/AndroidDeviceProfileSelector/Private/AndroidDeviceProfileSelector.cpp
+++ b/Engine/Plugins/Runtime/AndroidDeviceProfileSelector/Source/AndroidDeviceProfileSelector/Private/AndroidDeviceProfileSelector.cpp
@@ -224,7 +224,7 @@ FString FAndroidDeviceProfileSelector::FindMatchingProfile(const FString& GPUFam
 
 					FSHAHash SourceHash;
 					FSHA1::HashBuffer(TCHAR_TO_ANSI(*HashInputString), HashInputString.Len(), SourceHash.Hash);
-					if (SourceHash.ToString() != MatchHashString)
+					if (SourceHash.ToString() != MatchHashString.ToUpper())
 					{
 						bFoundMatch = false;
 					}
