Add optional properties map to device profile selector module.
This allows Device profile selector interface to be queried for platform specific properties.
Device preview can also set simulated values.
tidy up affected PIE preview and CheckAndroidDeviceProfile code.

diff --git a/Engine/Plugins/Runtime/AndroidDeviceProfileSelector/Source/AndroidDeviceProfileSelector/Private/AndroidDeviceProfileSelector.cpp b/Engine/Plugins/Runtime/AndroidDeviceProfileSelector/Source/AndroidDeviceProfileSelector/Private/AndroidDeviceProfileSelector.cpp
--- a/Engine/Plugins/Runtime/AndroidDeviceProfileSelector/Source/AndroidDeviceProfileSelector/Private/AndroidDeviceProfileSelector.cpp
+++ b/Engine/Plugins/Runtime/AndroidDeviceProfileSelector/Source/AndroidDeviceProfileSelector/Private/AndroidDeviceProfileSelector.cpp
@@ -36,10 +36,41 @@ static UAndroidDeviceProfileMatchingRules* GetAndroidDeviceProfileMatchingRules(
 	return Rules;
 }
 
-FString FAndroidDeviceProfileSelector::FindMatchingProfile(const FString& GPUFamily, const FString& GLVersion, const FString& AndroidVersion, const FString& DeviceMake, const FString& DeviceModel, const FString& DeviceBuildNumber, const FString& VulkanAvailable, const FString& VulkanVersion, const FString& UsingHoudini, const FString& Hardware, const FString& Chipset, const FString& ProfileName)
+TMap<FName, FString> FAndroidDeviceProfileSelector::SelectorProperties;
+
+void FAndroidDeviceProfileSelector::VerifySelectorParams()
+{
+	check(SelectorProperties.Find(FAndroidProfileSelectorSourceProperties::SRC_GPUFamily));
+	check(SelectorProperties.Find(FAndroidProfileSelectorSourceProperties::SRC_GLVersion));
+	check(SelectorProperties.Find(FAndroidProfileSelectorSourceProperties::SRC_AndroidVersion));
+	check(SelectorProperties.Find(FAndroidProfileSelectorSourceProperties::SRC_DeviceMake));
+	check(SelectorProperties.Find(FAndroidProfileSelectorSourceProperties::SRC_DeviceModel));
+	check(SelectorProperties.Find(FAndroidProfileSelectorSourceProperties::SRC_DeviceBuildNumber));
+	check(SelectorProperties.Find(FAndroidProfileSelectorSourceProperties::SRC_VulkanVersion));
+	check(SelectorProperties.Find(FAndroidProfileSelectorSourceProperties::SRC_UsingHoudini));
+	check(SelectorProperties.Find(FAndroidProfileSelectorSourceProperties::SRC_VulkanAvailable));
+	check(SelectorProperties.Find(FAndroidProfileSelectorSourceProperties::SRC_Hardware));
+	check(SelectorProperties.Find(FAndroidProfileSelectorSourceProperties::SRC_Chipset));
+	check(SelectorProperties.Find(FAndroidProfileSelectorSourceProperties::SRC_TotalPhysicalGB));
+}
+
+FString FAndroidDeviceProfileSelector::FindMatchingProfile(const FString& FallbackProfileName)
 {
-	FString OutProfileName = ProfileName;
+	FString OutProfileName = FallbackProfileName;
 	FString CommandLine = FCommandLine::Get();
+	
+	FString GPUFamily			= SelectorProperties.FindChecked(FAndroidProfileSelectorSourceProperties::SRC_GPUFamily);
+	FString GLVersion			= SelectorProperties.FindChecked(FAndroidProfileSelectorSourceProperties::SRC_GLVersion);
+	FString AndroidVersion		= SelectorProperties.FindChecked(FAndroidProfileSelectorSourceProperties::SRC_AndroidVersion);
+	FString DeviceMake			= SelectorProperties.FindChecked(FAndroidProfileSelectorSourceProperties::SRC_DeviceMake);
+	FString DeviceModel			= SelectorProperties.FindChecked(FAndroidProfileSelectorSourceProperties::SRC_DeviceModel);
+	FString DeviceBuildNumber	= SelectorProperties.FindChecked(FAndroidProfileSelectorSourceProperties::SRC_DeviceBuildNumber);
+	FString VulkanVersion		= SelectorProperties.FindChecked(FAndroidProfileSelectorSourceProperties::SRC_VulkanVersion);
+	FString UsingHoudini		= SelectorProperties.FindChecked(FAndroidProfileSelectorSourceProperties::SRC_UsingHoudini);
+	FString VulkanAvailable		= SelectorProperties.FindChecked(FAndroidProfileSelectorSourceProperties::SRC_VulkanAvailable);
+	FString Hardware			= SelectorProperties.FindChecked(FAndroidProfileSelectorSourceProperties::SRC_Hardware);
+	FString Chipset				= SelectorProperties.FindChecked(FAndroidProfileSelectorSourceProperties::SRC_Chipset);
+	FString TotalPhysicalGB		= SelectorProperties.FindChecked(FAndroidProfileSelectorSourceProperties::SRC_TotalPhysicalGB);
 
 	for (const FProfileMatch& Profile : GetAndroidDeviceProfileMatchingRules()->MatchProfile)
 	{
@@ -52,46 +83,46 @@ FString FAndroidDeviceProfileSelector::FindMatchingProfile(const FString& GPUFam
 			FString MatchString = Item.MatchString;
 			switch (Item.SourceType)
 			{
-			case SRC_PreviousRegexMatch:
+			case ESourceType::SRC_PreviousRegexMatch:
 				SourceString = &PreviousRegexMatch;
 				break;
-			case SRC_GpuFamily:
+			case ESourceType::SRC_GpuFamily:
 				SourceString = &GPUFamily;
 				break;
-			case SRC_GlVersion:
+			case ESourceType::SRC_GlVersion:
 				SourceString = &GLVersion;
 				break;
-			case SRC_AndroidVersion:
+			case ESourceType::SRC_AndroidVersion:
 				SourceString = &AndroidVersion;
 				break;
-			case SRC_DeviceMake:
+			case ESourceType::SRC_DeviceMake:
 				SourceString = &DeviceMake;
 				break;
-			case SRC_DeviceModel:
+			case ESourceType::SRC_DeviceModel:
 				SourceString = &DeviceModel;
 				break;
-			case SRC_DeviceBuildNumber:
+			case ESourceType::SRC_DeviceBuildNumber:
 				SourceString = &DeviceBuildNumber;
 				break;
-			case SRC_VulkanVersion:
+			case ESourceType::SRC_VulkanVersion:
 				SourceString = &VulkanVersion;
 				break;
-			case SRC_UsingHoudini:
+			case ESourceType::SRC_UsingHoudini:
 				SourceString = &UsingHoudini;
 				break;
-			case SRC_VulkanAvailable:
+			case ESourceType::SRC_VulkanAvailable:
 				SourceString = &VulkanAvailable;
 				break;
-			case SRC_CommandLine:
+			case ESourceType::SRC_CommandLine:
 				SourceString = &CommandLine;
 				break;
-			case SRC_Hardware:
+			case ESourceType::SRC_Hardware:
 				SourceString = &Hardware;
 				break;
-			case SRC_Chipset:
+			case ESourceType::SRC_Chipset:
 				SourceString = &Chipset;
 				break;
-			case SRC_ConfigRuleVar:
+			case ESourceType::SRC_ConfigRuleVar:
 				{
 					// expected Matchstring contents for configrulevar is "configrule_varname|matchstring"
 					// sourcestring is set to the configrule variable content. 
diff --git a/Engine/Plugins/Runtime/AndroidDeviceProfileSelector/Source/AndroidDeviceProfileSelector/Private/AndroidDeviceProfileSelectorModule.cpp b/Engine/Plugins/Runtime/AndroidDeviceProfileSelector/Source/AndroidDeviceProfileSelector/Private/AndroidDeviceProfileSelectorModule.cpp
--- a/Engine/Plugins/Runtime/AndroidDeviceProfileSelector/Source/AndroidDeviceProfileSelector/Private/AndroidDeviceProfileSelectorModule.cpp
+++ b/Engine/Plugins/Runtime/AndroidDeviceProfileSelector/Source/AndroidDeviceProfileSelector/Private/AndroidDeviceProfileSelectorModule.cpp
@@ -25,40 +25,38 @@ const FString FAndroidDeviceProfileSelectorModule::GetRuntimeDeviceProfileName()
 	return FString();
 }
 
-const FString FAndroidDeviceProfileSelectorModule::GetDeviceProfileName(const TMap<FString, FString>& DeviceParameters)
+const FString FAndroidDeviceProfileSelectorModule::GetDeviceProfileName()
 {
 	FString ProfileName; 
 
-	// Pull out required device parameters:
-	FString GPUFamily = DeviceParameters.FindChecked("GPUFamily");
-	FString GLVersion = DeviceParameters.FindChecked("GLVersion");
-	FString VulkanAvailable = DeviceParameters.FindChecked("VulkanAvailable");
-	FString VulkanVersion = DeviceParameters.FindChecked("VulkanVersion");
-	FString AndroidVersion = DeviceParameters.FindChecked("AndroidVersion");
-	FString DeviceMake = DeviceParameters.FindChecked("DeviceMake");
-	FString DeviceModel = DeviceParameters.FindChecked("DeviceModel");
-	FString DeviceBuildNumber = DeviceParameters.FindChecked("DeviceBuildNumber");
-	FString UsingHoudini = DeviceParameters.FindChecked("UsingHoudini");
-	FString Hardware = DeviceParameters.FindChecked("Hardware");
-	FString Chipset = DeviceParameters.FindChecked("Chipset");
+	// ensure SelectorProperties does actually contain our parameters
+	check(FAndroidDeviceProfileSelector::GetSelectorProperties().Num() > 0);
 
 	UE_LOG(LogAndroidDPSelector, Log, TEXT("Checking %d rules from DeviceProfile ini file."), FAndroidDeviceProfileSelector::GetNumProfiles() );
 	UE_LOG(LogAndroidDPSelector, Log, TEXT("  Default profile: %s"), *ProfileName);
-	UE_LOG(LogAndroidDPSelector, Log, TEXT("  GpuFamily: %s"), *GPUFamily);
-	UE_LOG(LogAndroidDPSelector, Log, TEXT("  GlVersion: %s"), *GLVersion);
-	UE_LOG(LogAndroidDPSelector, Log, TEXT("  VulkanAvailable: %s"), *VulkanAvailable);
-	UE_LOG(LogAndroidDPSelector, Log, TEXT("  VulkanVersion: %s"), *VulkanVersion);
-	UE_LOG(LogAndroidDPSelector, Log, TEXT("  AndroidVersion: %s"), *AndroidVersion);
-	UE_LOG(LogAndroidDPSelector, Log, TEXT("  DeviceMake: %s"), *DeviceMake);
-	UE_LOG(LogAndroidDPSelector, Log, TEXT("  DeviceModel: %s"), *DeviceModel);
-	UE_LOG(LogAndroidDPSelector, Log, TEXT("  DeviceBuildNumber: %s"), *DeviceBuildNumber);
-	UE_LOG(LogAndroidDPSelector, Log, TEXT("  UsingHoudini: %s"), *UsingHoudini);
-	UE_LOG(LogAndroidDPSelector, Log, TEXT("  Hardware: %s"), *Hardware);
-	UE_LOG(LogAndroidDPSelector, Log, TEXT("  Chipset: %s"), *Chipset);
+	for (const TTuple<FName,FString>& MapIt : FAndroidDeviceProfileSelector::GetSelectorProperties())
+	{
+		UE_LOG(LogAndroidDPSelector, Log, TEXT("  %s: %s"), *MapIt.Key.ToString(), *MapIt.Value);
+	}
 
-	ProfileName = FAndroidDeviceProfileSelector::FindMatchingProfile(GPUFamily, GLVersion, AndroidVersion, DeviceMake, DeviceModel, DeviceBuildNumber, VulkanAvailable, VulkanVersion, UsingHoudini, Hardware, Chipset, ProfileName);
+	ProfileName = FAndroidDeviceProfileSelector::FindMatchingProfile(ProfileName);
 
 	UE_LOG(LogAndroidDPSelector, Log, TEXT("Selected Device Profile: [%s]"), *ProfileName);
 
 	return ProfileName;
 }
+
+bool FAndroidDeviceProfileSelectorModule::GetSelectorPropertyValue(const FName& PropertyType, FString& PropertyValueOUT)
+{
+	if (const FString* Found = FAndroidDeviceProfileSelector::GetSelectorProperties().Find(PropertyType))
+	{
+		PropertyValueOUT = *Found;
+		return true;
+	}
+	return false;
+}
+
+void FAndroidDeviceProfileSelectorModule::SetSelectorProperties(const TMap<FName, FString>& SelectorPropertiesIn)
+{
+	FAndroidDeviceProfileSelector::SetSelectorProperties(SelectorPropertiesIn);
+}
diff --git a/Engine/Plugins/Runtime/AndroidDeviceProfileSelector/Source/AndroidDeviceProfileSelector/Private/AndroidDeviceProfileSelectorModule.h b/Engine/Plugins/Runtime/AndroidDeviceProfileSelector/Source/AndroidDeviceProfileSelector/Private/AndroidDeviceProfileSelectorModule.h
--- a/Engine/Plugins/Runtime/AndroidDeviceProfileSelector/Source/AndroidDeviceProfileSelector/Private/AndroidDeviceProfileSelectorModule.h
+++ b/Engine/Plugins/Runtime/AndroidDeviceProfileSelector/Source/AndroidDeviceProfileSelector/Private/AndroidDeviceProfileSelectorModule.h
@@ -7,6 +7,7 @@
 
 /**
  * Implements the Android Device Profile Selector module.
+ * Used to emulate an android device's IDeviceProfileSelectorModule behaviour.
  */
 class FAndroidDeviceProfileSelectorModule
 	: public IDeviceProfileSelectorModule
@@ -14,10 +15,13 @@ class FAndroidDeviceProfileSelectorModule
 public:
 
 	//~ Begin IDeviceProfileSelectorModule Interface
-	virtual const FString GetDeviceProfileName(const TMap<FString, FString>& DeviceParameters) override;
+	virtual const FString GetDeviceProfileName() override;
 	virtual const FString GetRuntimeDeviceProfileName() override;
-	//~ End IDeviceProfileSelectorModule Interface
 
+	// Set the device parameters this selector will use.
+	virtual void SetSelectorProperties(const TMap<FName, FString>& SelectorProperties) override;
+	virtual bool GetSelectorPropertyValue(const FName& PropertyType, FString& PropertyValueOUT) override;
+	//~ End IDeviceProfileSelectorModule Interface
 
 	//~ Begin IModuleInterface Interface
 	virtual void StartupModule() override;
diff --git a/Engine/Plugins/Runtime/AndroidDeviceProfileSelector/Source/AndroidDeviceProfileSelector/Public/AndroidDeviceProfileSelector.h b/Engine/Plugins/Runtime/AndroidDeviceProfileSelector/Source/AndroidDeviceProfileSelector/Public/AndroidDeviceProfileSelector.h
--- a/Engine/Plugins/Runtime/AndroidDeviceProfileSelector/Source/AndroidDeviceProfileSelector/Public/AndroidDeviceProfileSelector.h
+++ b/Engine/Plugins/Runtime/AndroidDeviceProfileSelector/Source/AndroidDeviceProfileSelector/Public/AndroidDeviceProfileSelector.h
@@ -4,9 +4,31 @@
 
 #include "CoreMinimal.h"
 
+namespace FAndroidProfileSelectorSourceProperties
+{
+	static FName SRC_GPUFamily(TEXT("SRC_GPUFamily"));
+	static FName SRC_GLVersion(TEXT("SRC_GLVersion"));
+	static FName SRC_VulkanAvailable(TEXT("SRC_VulkanAvailable"));
+	static FName SRC_VulkanVersion(TEXT("SRC_VulkanVersion"));
+	static FName SRC_AndroidVersion(TEXT("SRC_AndroidVersion"));
+	static FName SRC_DeviceMake(TEXT("SRC_DeviceMake"));
+	static FName SRC_DeviceModel(TEXT("SRC_DeviceModel"));
+	static FName SRC_DeviceBuildNumber(TEXT("SRC_DeviceBuildNumber"));
+	static FName SRC_UsingHoudini(TEXT("SRC_UsingHoudini"));
+	static FName SRC_Hardware(TEXT("SRC_Hardware"));
+	static FName SRC_Chipset(TEXT("SRC_Chipset"));
+	static FName SRC_TotalPhysicalGB(TEXT("SRC_TotalPhysicalGB"));
+};
+
 class ANDROIDDEVICEPROFILESELECTOR_API FAndroidDeviceProfileSelector
 {
+	// Container of various device properties used for device profile matching.
+	static TMap<FName, FString> SelectorProperties;
+	static void VerifySelectorParams();
 public:
-	static FString FindMatchingProfile(const FString& GPUFamily, const FString& GLVersion, const FString& AndroidVersion, const FString& DeviceMake, const FString& DeviceModel, const FString& DeviceBuildNumber, const FString& VulkanAvailable, const FString& VulkanVersion, const FString& UsingHoudini, const FString& Hardware, const FString& Chipset, const FString& ProfileName);
+
+	static FString FindMatchingProfile(const FString& FallbackProfileName);
 	static int32 GetNumProfiles();
+	static void SetSelectorProperties(const TMap<FName, FString>& Params) { SelectorProperties = Params; VerifySelectorParams(); }
+	static const TMap<FName, FString>& GetSelectorProperties() { return SelectorProperties; }
 };
diff --git a/Engine/Plugins/Runtime/AndroidDeviceProfileSelector/Source/AndroidDeviceProfileSelectorRuntime/Private/AndroidDeviceProfileSelectorRuntimeModule.cpp b/Engine/Plugins/Runtime/AndroidDeviceProfileSelector/Source/AndroidDeviceProfileSelectorRuntime/Private/AndroidDeviceProfileSelectorRuntimeModule.cpp
--- a/Engine/Plugins/Runtime/AndroidDeviceProfileSelector/Source/AndroidDeviceProfileSelectorRuntime/Private/AndroidDeviceProfileSelectorRuntimeModule.cpp
+++ b/Engine/Plugins/Runtime/AndroidDeviceProfileSelector/Source/AndroidDeviceProfileSelectorRuntime/Private/AndroidDeviceProfileSelectorRuntimeModule.cpp
@@ -19,6 +19,60 @@ void FAndroidDeviceProfileSelectorRuntimeModule::ShutdownModule()
 {
 }
 
+// Build the selector params for the current device.
+static const TMap<FName,FString>& GetDeviceSelectorParams()
+{
+	static bool bInitialized = false;
+	static TMap<FName, FString> AndroidParams;
+	if(!bInitialized)
+	{
+		bInitialized = true;
+		auto GetParam = [](const FString& DefaultParam, const TCHAR* ConfRuleVarName)
+		{
+#if PLATFORM_ANDROID
+			if (FString* ConfRuleVarValue = FAndroidMisc::GetConfigRulesVariable(ConfRuleVarName))
+			{
+				return *ConfRuleVarValue;
+			}
+#endif
+			return DefaultParam;
+		};
+
+#if !(PLATFORM_ANDROID_X86 || PLATFORM_ANDROID_X64)
+		// Not running an Intel libUnreal.so with Houdini library present means we're emulated
+		bool bUsingHoudini = (access("/system/lib/libhoudini.so", F_OK) != -1);
+#else
+		bool bUsingHoudini = false;
+#endif
+
+		FPlatformMemoryStats Stats = FPlatformMemory::GetStats();
+		// this is used in the same way as PlatformMemoryBucket
+		// which on Android has a different rounding algo. See GenericPlatformMemory::GetMemorySizeBucket.
+		uint64 MemoryBucketRoundingAddition = 384;
+#if PLATFORM_ANDROID
+		if (FString* MemoryBucketRoundingAdditionVar = FAndroidMisc::GetConfigRulesVariable(TEXT("MemoryBucketRoundingAddition")))
+		{
+			MemoryBucketRoundingAddition = FCString::Atoi64(**MemoryBucketRoundingAdditionVar);
+		}
+#endif
+		uint32 TotalPhysicalGB = (uint32)((Stats.TotalPhysical + MemoryBucketRoundingAddition * 1024 * 1024 - 1) / 1024 / 1024 / 1024);
+
+		AndroidParams.Add(FAndroidProfileSelectorSourceProperties::SRC_GPUFamily, GetParam(FAndroidMisc::GetGPUFamily(), TEXT("gpu")));
+		AndroidParams.Add(FAndroidProfileSelectorSourceProperties::SRC_GLVersion, FAndroidMisc::GetGLVersion());
+		AndroidParams.Add(FAndroidProfileSelectorSourceProperties::SRC_VulkanAvailable, FAndroidMisc::IsVulkanAvailable() ? TEXT("true") : TEXT("false"));
+		AndroidParams.Add(FAndroidProfileSelectorSourceProperties::SRC_VulkanVersion, FAndroidMisc::GetVulkanVersion());
+		AndroidParams.Add(FAndroidProfileSelectorSourceProperties::SRC_AndroidVersion, FAndroidMisc::GetAndroidVersion());
+		AndroidParams.Add(FAndroidProfileSelectorSourceProperties::SRC_DeviceMake, FAndroidMisc::GetDeviceMake());
+		AndroidParams.Add(FAndroidProfileSelectorSourceProperties::SRC_DeviceModel, FAndroidMisc::GetDeviceModel());
+		AndroidParams.Add(FAndroidProfileSelectorSourceProperties::SRC_DeviceBuildNumber, FAndroidMisc::GetDeviceBuildNumber());
+		AndroidParams.Add(FAndroidProfileSelectorSourceProperties::SRC_UsingHoudini, bUsingHoudini ? TEXT("true") : TEXT("false"));
+		AndroidParams.Add(FAndroidProfileSelectorSourceProperties::SRC_Hardware, GetParam(FString(TEXT("unknown")), TEXT("hardware")));
+		AndroidParams.Add(FAndroidProfileSelectorSourceProperties::SRC_Chipset, GetParam(FString(TEXT("unknown")), TEXT("chipset")));
+		AndroidParams.Add(FAndroidProfileSelectorSourceProperties::SRC_TotalPhysicalGB, FString::Printf(TEXT("%d"), TotalPhysicalGB));
+	}
+	return AndroidParams;
+}
+
 FString const FAndroidDeviceProfileSelectorRuntimeModule::GetRuntimeDeviceProfileName()
 {
 	static FString ProfileName;
@@ -31,51 +85,20 @@ FString const FAndroidDeviceProfileSelectorRuntimeModule::GetRuntimeDeviceProfil
 		{
 			ProfileName = FPlatformProperties::PlatformName();
 		}
+		const TMap<FName, FString>& AndroidParams = GetDeviceSelectorParams();
 
-		FString GPUFamily = FAndroidMisc::GetGPUFamily();
-		FString GLVersion = FAndroidMisc::GetGLVersion();
-
-		FString VulkanVersion = FAndroidMisc::GetVulkanVersion();
-		FString VulkanAvailable = FAndroidMisc::IsVulkanAvailable() ? TEXT("true") : TEXT("false");
-		FString AndroidVersion = FAndroidMisc::GetAndroidVersion();
-		FString DeviceMake = FAndroidMisc::GetDeviceMake();
-		FString DeviceModel = FAndroidMisc::GetDeviceModel();
-		FString DeviceBuildNumber = FAndroidMisc::GetDeviceBuildNumber();
-		
-		FString Hardware = FString(TEXT("unknown"));
-		FString Chipset = FString(TEXT("unknown"));
-		FString *HardwareLookup = FAndroidMisc::GetConfigRulesVariable(TEXT("hardware"));
-		if (HardwareLookup != nullptr)
-		{
-			Hardware = *HardwareLookup;
-		}
-		FString *ChipsetLookup = FAndroidMisc::GetConfigRulesVariable(TEXT("chipset"));
-		if (ChipsetLookup != nullptr)
-		{
-			Chipset = *ChipsetLookup;
-		}
-
-#if !(PLATFORM_ANDROID_X86 || PLATFORM_ANDROID_X64)
-		// Not running an Intel libUnreal.so with Houdini library present means we're emulated
-		bool bUsingHoudini = (access("/system/lib/libhoudini.so", F_OK) != -1);
-#else
-		bool bUsingHoudini = false;
-#endif
-		FString UsingHoudini = bUsingHoudini ? TEXT("true") : TEXT("false");
+		FAndroidDeviceProfileSelector::SetSelectorProperties(AndroidParams);
 
 		UE_LOG(LogAndroid, Log, TEXT("Checking %d rules from DeviceProfile ini file."), FAndroidDeviceProfileSelector::GetNumProfiles() );
 		UE_LOG(LogAndroid, Log, TEXT("  Default profile: %s"), * ProfileName);
-		UE_LOG(LogAndroid, Log, TEXT("  GpuFamily: %s"), *GPUFamily);
-		UE_LOG(LogAndroid, Log, TEXT("  GlVersion: %s"), *GLVersion);
-		UE_LOG(LogAndroid, Log, TEXT("  VulkanAvailable: %s"), *VulkanAvailable);
-		UE_LOG(LogAndroid, Log, TEXT("  VulkanVersion: %s"), *VulkanVersion);
-		UE_LOG(LogAndroid, Log, TEXT("  AndroidVersion: %s"), *AndroidVersion);
-		UE_LOG(LogAndroid, Log, TEXT("  DeviceMake: %s"), *DeviceMake);
-		UE_LOG(LogAndroid, Log, TEXT("  DeviceModel: %s"), *DeviceModel);
-		UE_LOG(LogAndroid, Log, TEXT("  DeviceBuildNumber: %s"), *DeviceBuildNumber);
-		UE_LOG(LogAndroid, Log, TEXT("  UsingHoudini: %s"), *UsingHoudini);
-		UE_LOG(LogAndroid, Log, TEXT("  Hardware: %s"), *Hardware);
-		UE_LOG(LogAndroid, Log, TEXT("  Chipset: %s"), *Chipset);
+		UE_LOG(LogAndroid, Log, TEXT("  Android selector params: "));
+		for(auto& MapIt : AndroidParams)
+		{
+			UE_LOG(LogAndroid, Log, TEXT("  %s: %s"), *MapIt.Key.ToString(), *MapIt.Value);
+		}
+
+		const FString& DeviceMake = AndroidParams.FindChecked(FAndroidProfileSelectorSourceProperties::SRC_DeviceMake);
+		const FString& DeviceModel = AndroidParams.FindChecked(FAndroidProfileSelectorSourceProperties::SRC_DeviceModel);
 
 		CheckForJavaSurfaceViewWorkaround(DeviceMake, DeviceModel);
 
@@ -89,7 +112,7 @@ FString const FAndroidDeviceProfileSelectorRuntimeModule::GetRuntimeDeviceProfil
 		else
 		{
 			// Find a match with the DeviceProfiles matching rules
-			ProfileName = FAndroidDeviceProfileSelector::FindMatchingProfile(GPUFamily, GLVersion, AndroidVersion, DeviceMake, DeviceModel, DeviceBuildNumber, VulkanAvailable, VulkanVersion, UsingHoudini, Hardware, Chipset, ProfileName);
+			ProfileName = FAndroidDeviceProfileSelector::FindMatchingProfile(ProfileName);
 			UE_LOG(LogAndroid, Log, TEXT("Selected Device Profile: [%s]"), *ProfileName);
 		}
 	}
@@ -97,6 +120,16 @@ FString const FAndroidDeviceProfileSelectorRuntimeModule::GetRuntimeDeviceProfil
 	return ProfileName;
 }
 
+bool FAndroidDeviceProfileSelectorRuntimeModule::GetSelectorPropertyValue(const FName& PropertyType, FString& PropertyValueOUT)
+{
+	if (const FString* Found = GetDeviceSelectorParams().Find(PropertyType))
+	{
+		PropertyValueOUT = *Found;
+		return true;
+	}
+	return false;
+}
+
 void FAndroidDeviceProfileSelectorRuntimeModule::CheckForJavaSurfaceViewWorkaround(const FString& DeviceMake, const FString& DeviceModel) const
 {
 #if USE_ANDROID_JNI
diff --git a/Engine/Plugins/Runtime/AndroidDeviceProfileSelector/Source/AndroidDeviceProfileSelectorRuntime/Private/AndroidDeviceProfileSelectorRuntimeModule.h b/Engine/Plugins/Runtime/AndroidDeviceProfileSelector/Source/AndroidDeviceProfileSelectorRuntime/Private/AndroidDeviceProfileSelectorRuntimeModule.h
--- a/Engine/Plugins/Runtime/AndroidDeviceProfileSelector/Source/AndroidDeviceProfileSelectorRuntime/Private/AndroidDeviceProfileSelectorRuntimeModule.h
+++ b/Engine/Plugins/Runtime/AndroidDeviceProfileSelector/Source/AndroidDeviceProfileSelectorRuntime/Private/AndroidDeviceProfileSelectorRuntimeModule.h
@@ -15,6 +15,7 @@ public:
 
 	//~ Begin IDeviceProfileSelectorModule Interface
 	virtual const FString GetRuntimeDeviceProfileName() override;
+	virtual bool GetSelectorPropertyValue(const FName& PropertyType, FString& PropertyValueOUT) override;
 	//~ End IDeviceProfileSelectorModule Interface
 
 
diff --git a/Engine/Source/Editor/DeviceProfileServices/Private/CheckAndroidDeviceProfileCommandlet.cpp b/Engine/Source/Editor/DeviceProfileServices/Private/CheckAndroidDeviceProfileCommandlet.cpp
--- a/Engine/Source/Editor/DeviceProfileServices/Private/CheckAndroidDeviceProfileCommandlet.cpp
+++ b/Engine/Source/Editor/DeviceProfileServices/Private/CheckAndroidDeviceProfileCommandlet.cpp
@@ -16,25 +16,27 @@ int32 UCheckAndroidDeviceProfileCommandlet::Main(const FString& RawCommandLine)
 	IDeviceProfileSelectorModule* AndroidDeviceProfileSelector = FModuleManager::LoadModulePtr<IDeviceProfileSelectorModule>("AndroidDeviceProfileSelector");
 	if (ensure(AndroidDeviceProfileSelector != nullptr))
 	{
-		TMap<FString, FString> DeviceParameters;
-		DeviceParameters.Add(TEXT("GPUFamily"), Params.FindRef(TEXT("GPUFamily")));
-		DeviceParameters.Add(TEXT("GLVersion"), Params.FindRef(TEXT("GLVersion")));
-		DeviceParameters.Add(TEXT("VulkanAvailable"), Params.FindRef(TEXT("VulkanAvailable")));
-		DeviceParameters.Add(TEXT("VulkanVersion"), Params.FindRef(TEXT("VulkanVersion")));
-		DeviceParameters.Add(TEXT("AndroidVersion"), Params.FindRef(TEXT("AndroidVersion")));
-		DeviceParameters.Add(TEXT("DeviceMake"),
+		TMap<FName, FString> DeviceParameters;
+		DeviceParameters.Add(FName(TEXT("SRC_GPUFamily")), Params.FindRef(TEXT("GPUFamily")));
+		DeviceParameters.Add(FName(TEXT("SRC_GLVersion")), Params.FindRef(TEXT("GLVersion")));
+		DeviceParameters.Add(FName(TEXT("SRC_VulkanAvailable")), Params.FindRef(TEXT("VulkanAvailable")));
+		DeviceParameters.Add(FName(TEXT("SRC_VulkanVersion")), Params.FindRef(TEXT("VulkanVersion")));
+		DeviceParameters.Add(FName(TEXT("SRC_AndroidVersion")), Params.FindRef(TEXT("AndroidVersion")));
+		DeviceParameters.Add(FName(TEXT("SRC_DeviceMake")),
 			Tokens.Num() == 2 ? Tokens[0] :
 			Params.FindRef(TEXT("DeviceMake")));
-		DeviceParameters.Add(TEXT("DeviceModel"),
+		DeviceParameters.Add(FName(TEXT("SRC_DeviceModel")),
 			Tokens.Num() == 1 ? Tokens[0] :
 			Tokens.Num() == 2 ? Tokens[1] : 
 			Params.FindRef(TEXT("DeviceModel")));
-		DeviceParameters.Add(TEXT("DeviceBuildNumber"), Params.FindRef(TEXT("DeviceBuildNumber")));
-		DeviceParameters.Add(TEXT("UsingHoudini"), Params.FindRef(TEXT("UsingHoudini")));
-		DeviceParameters.Add(TEXT("Hardware"), Params.FindRef(TEXT("Hardware")));
-		DeviceParameters.Add(TEXT("Chipset"), Params.FindRef(TEXT("Chipset")));
+		DeviceParameters.Add(FName(TEXT("SRC_DeviceBuildNumber")), Params.FindRef(TEXT("DeviceBuildNumber")));
+		DeviceParameters.Add(FName(TEXT("SRC_UsingHoudini")), Params.FindRef(TEXT("UsingHoudini")));
+		DeviceParameters.Add(FName(TEXT("SRC_Hardware")), Params.FindRef(TEXT("Hardware")));
+		DeviceParameters.Add(FName(TEXT("SRC_Chipset")), Params.FindRef(TEXT("Chipset")));
+		DeviceParameters.Add(FName(TEXT("SRC_TotalPhysicalGB")), Params.FindRef(TEXT("TotalPhysicalGB")));
 
-		FString ProfileName = AndroidDeviceProfileSelector->GetDeviceProfileName(DeviceParameters);
+		AndroidDeviceProfileSelector->SetSelectorProperties(DeviceParameters);
+		FString ProfileName = AndroidDeviceProfileSelector->GetDeviceProfileName();
 
 		UE_LOG(LogCheckAndroidDeviceProfile, Display, TEXT("Selected Device Profile: %s"), *ProfileName);
 	}
diff --git a/Engine/Source/Editor/PIEPreviewDeviceProfileSelector/Private/PIEPreviewDevice.cpp b/Engine/Source/Editor/PIEPreviewDeviceProfileSelector/Private/PIEPreviewDevice.cpp
--- a/Engine/Source/Editor/PIEPreviewDeviceProfileSelector/Private/PIEPreviewDevice.cpp
+++ b/Engine/Source/Editor/PIEPreviewDeviceProfileSelector/Private/PIEPreviewDevice.cpp
@@ -424,20 +424,22 @@ FString FPIEPreviewDevice::GetProfile() const
 			{
 				FPIEAndroidDeviceProperties& AndroidProperties = DeviceSpecs->AndroidProperties;
 
-				TMap<FString, FString> DeviceParameters;
-				DeviceParameters.Add("GPUFamily", AndroidProperties.GPUFamily);
-				DeviceParameters.Add("GLVersion", AndroidProperties.GLVersion);
-				DeviceParameters.Add("VulkanAvailable", AndroidProperties.VulkanAvailable ? "true" : "false");
-				DeviceParameters.Add("VulkanVersion", AndroidProperties.VulkanVersion);
-				DeviceParameters.Add("AndroidVersion", AndroidProperties.AndroidVersion);
-				DeviceParameters.Add("DeviceMake", AndroidProperties.DeviceMake);
-				DeviceParameters.Add("DeviceModel", AndroidProperties.DeviceModel);
-				DeviceParameters.Add("DeviceBuildNumber", AndroidProperties.DeviceBuildNumber);
-				DeviceParameters.Add("UsingHoudini", AndroidProperties.UsingHoudini ? "true" : "false");
-				DeviceParameters.Add("Hardware", AndroidProperties.Hardware);
-				DeviceParameters.Add("Chipset", AndroidProperties.Chipset);
-
-				FString PIEProfileName = AndroidDeviceProfileSelector->GetDeviceProfileName(DeviceParameters);
+				TMap<FName, FString> DeviceParameters;
+				DeviceParameters.Add(FName(TEXT("SRC_GPUFamily")), AndroidProperties.GPUFamily);
+				DeviceParameters.Add(FName(TEXT("SRC_GLVersion")), AndroidProperties.GLVersion);
+				DeviceParameters.Add(FName(TEXT("SRC_VulkanAvailable")), AndroidProperties.VulkanAvailable ? "true" : "false");
+				DeviceParameters.Add(FName(TEXT("SRC_VulkanVersion")), AndroidProperties.VulkanVersion);
+				DeviceParameters.Add(FName(TEXT("SRC_AndroidVersion")), AndroidProperties.AndroidVersion);
+				DeviceParameters.Add(FName(TEXT("SRC_DeviceMake")), AndroidProperties.DeviceMake);
+				DeviceParameters.Add(FName(TEXT("SRC_DeviceModel")), AndroidProperties.DeviceModel);
+				DeviceParameters.Add(FName(TEXT("SRC_DeviceBuildNumber")), AndroidProperties.DeviceBuildNumber);
+				DeviceParameters.Add(FName(TEXT("SRC_UsingHoudini")), AndroidProperties.UsingHoudini ? "true" : "false");
+				DeviceParameters.Add(FName(TEXT("SRC_Hardware")), AndroidProperties.Hardware);
+				DeviceParameters.Add(FName(TEXT("SRC_Chipset")), AndroidProperties.Chipset);
+				DeviceParameters.Add(FName(TEXT("SRC_TotalPhysicalGB")), AndroidProperties.TotalPhysicalGB);
+
+				AndroidDeviceProfileSelector->SetSelectorProperties(DeviceParameters);
+				FString PIEProfileName = AndroidDeviceProfileSelector->GetDeviceProfileName();
 				if (!PIEProfileName.IsEmpty())
 				{
 					Profile = PIEProfileName;
diff --git a/Engine/Source/Runtime/Engine/Public/IDeviceProfileSelectorModule.h b/Engine/Source/Runtime/Engine/Public/IDeviceProfileSelectorModule.h
--- a/Engine/Source/Runtime/Engine/Public/IDeviceProfileSelectorModule.h
+++ b/Engine/Source/Runtime/Engine/Public/IDeviceProfileSelectorModule.h
@@ -30,8 +30,21 @@ public:
 	* @param DeviceParameters	A map of parameters to be used by device profile logic.
 	* @return The name of the device profile to use for this session
 	*/
-	virtual const FString GetDeviceProfileName(const TMap<FString, FString>& DeviceParameters) { return FString(); }
+	virtual const FString GetDeviceProfileName() { return FString(); }
 
+	/**
+	* Set or override the selector specific properties.
+	* @param SelectorProperties	A map of parameters to be used by device profile matching logic.
+	*/
+	virtual void SetSelectorProperties(const TMap<FName, FString>& SelectorProperties) { }
+
+	/*
+	* Find a custom profile selector property value.
+	* @Param PropertyType The information requested
+	* @Param PropertyValueOUT the value of PropertyType
+	* @return Whether the PropertyType was recognized by the profile selector.
+	*/
+	virtual bool GetSelectorPropertyValue(const FName& PropertyType, FString& PropertyValueOUT) { PropertyValueOUT.Reset(); return false; }
 
 	/**
 	 * Virtual destructor.
