[Partial Backout] - CL24348577
Original CL Desc
-----------------------------------------------------------------
Remove unnecessary private include path

#rnx

diff --git a/Engine/Plugins/Runtime/AndroidDeviceProfileSelector/Source/AndroidDeviceProfileCommandlets/AndroidDeviceProfileCommandlets.Build.cs b/Engine/Plugins/Runtime/AndroidDeviceProfileSelector/Source/AndroidDeviceProfileCommandlets/AndroidDeviceProfileCommandlets.Build.cs
--- a/Engine/Plugins/Runtime/AndroidDeviceProfileSelector/Source/AndroidDeviceProfileCommandlets/AndroidDeviceProfileCommandlets.Build.cs
+++ b/Engine/Plugins/Runtime/AndroidDeviceProfileSelector/Source/AndroidDeviceProfileCommandlets/AndroidDeviceProfileCommandlets.Build.cs
@@ -7,7 +7,7 @@ public class AndroidDeviceProfileCommandlets : ModuleRules
 {
 	public AndroidDeviceProfileCommandlets(ReadOnlyTargetRules Target) : base(Target)
 	{
-		PublicIncludePathModuleNames.Add("AndroidDeviceDetection");
+		PublicIncludePaths.Add("Developer/Android/AndroidDeviceDetection/Public");
 
 		PublicDependencyModuleNames.AddRange(
 			new string[]
