[Backout] - CL26558401 (partial backout)
Original CL Desc
-----------------------------------------------------------------
Resolve several module public include paths not using references

#rnx

diff --git a/Engine/Plugins/Runtime/AndroidDeviceProfileSelector/Source/AndroidDeviceProfileCommandlets/AndroidDeviceProfileCommandlets.Build.cs b/Engine/Plugins/Runtime/AndroidDeviceProfileSelector/Source/AndroidDeviceProfileCommandlets/AndroidDeviceProfileCommandlets.Build.cs
--- a/Engine/Plugins/Runtime/AndroidDeviceProfileSelector/Source/AndroidDeviceProfileCommandlets/AndroidDeviceProfileCommandlets.Build.cs
+++ b/Engine/Plugins/Runtime/AndroidDeviceProfileSelector/Source/AndroidDeviceProfileCommandlets/AndroidDeviceProfileCommandlets.Build.cs
@@ -1,24 +1,21 @@
 // Copyright Epic Games, Inc. All Rights Reserved.
+using EpicGames.Core;
+using System.IO;
 using UnrealBuildTool;
 
 public class AndroidDeviceProfileCommandlets : ModuleRules
 {
 	public AndroidDeviceProfileCommandlets(ReadOnlyTargetRules Target) : base(Target)
 	{
+		PublicIncludePaths.Add("Developer/Android/AndroidDeviceDetection/Public");
+
 		PublicDependencyModuleNames.AddRange(
 			new string[]
 			{
 				"Core",
 				"CoreUObject",
 			}
-		);
-
-		PublicIncludePathModuleNames.AddRange(
-			new string[]
-			{
-				"AndroidDeviceDetection",
-			}
-		);
+			);
 
 		PrivateDependencyModuleNames.AddRange(
 			new string[]
@@ -29,6 +26,6 @@ public class AndroidDeviceProfileCommandlets : ModuleRules
 				"AndroidDeviceProfileSelector",
 				"PIEPreviewDeviceSpecification"
 			}
-		);
+			);
 	}
 }
