Aspect Ratio Preview fixes and Safe Zones preview  implementation
Wrapped the main editor viewport in the middle of a 3x3 grid in order to restrict aspect ratio for preview platforms that request it.
By default, without preview or aspect ratio restriction, the middle cell will fill the entire screen.
Safe Zones are gathered from adb dumpsys windows display atm and made resolution independent

diff --git a/Engine/Plugins/Runtime/AndroidDeviceProfileSelector/Source/AndroidDeviceProfileSelector/Private/AndroidDeviceProfileSelectorModule.cpp b/Engine/Plugins/Runtime/AndroidDeviceProfileSelector/Source/AndroidDeviceProfileSelector/Private/AndroidDeviceProfileSelectorModule.cpp
--- a/Engine/Plugins/Runtime/AndroidDeviceProfileSelector/Source/AndroidDeviceProfileSelector/Private/AndroidDeviceProfileSelectorModule.cpp
+++ b/Engine/Plugins/Runtime/AndroidDeviceProfileSelector/Source/AndroidDeviceProfileSelector/Private/AndroidDeviceProfileSelectorModule.cpp
@@ -111,23 +111,12 @@ void FAndroidDeviceProfileSelectorModule::GetDeviceParametersFromJson(FString& J
 		OutDeviceParameters.Add(FName(TEXT("SRC_TotalPhysicalGB")), AndroidProperties.TotalPhysicalGB);
 		OutDeviceParameters.Add(FName(TEXT("SRC_HMDSystemName")), TEXT(""));
 		OutDeviceParameters.Add(FName(TEXT("SRC_SM5Available")), AndroidProperties.SM5Available ? "true" : "false");
-
-		FString Orientation;
-		bool bNeedPortrait = false;
-		GConfig->GetString(TEXT("/Script/AndroidRuntimeSettings.AndroidRuntimeSettings"), TEXT("Orientation"), Orientation, GEngineIni);
-		if (Orientation.ToLower().Equals("portrait") || Orientation.ToLower().Equals("reverseportrait") || Orientation.ToLower().Equals("sensorportrait"))
-		{
-			bNeedPortrait = true;
-		}
-
-		if (bNeedPortrait)
-		{
-			ConstrainedAspectRatio = float(DeviceSpecs.ResolutionX) / float(DeviceSpecs.ResolutionY);
-		}
-		else
-		{
-			ConstrainedAspectRatio = float(DeviceSpecs.ResolutionY) / float(DeviceSpecs.ResolutionX);
-		}
+		OutDeviceParameters.Add(FAndroidProfileSelectorSourceProperties::SRC_ResolutionX, FString::Printf(TEXT("%d"), DeviceSpecs.ResolutionX));
+		OutDeviceParameters.Add(FAndroidProfileSelectorSourceProperties::SRC_ResolutionY, FString::Printf(TEXT("%d"), DeviceSpecs.ResolutionY));
+		OutDeviceParameters.Add(FAndroidProfileSelectorSourceProperties::SRC_InsetsLeft, FString::Printf(TEXT("%f"), DeviceSpecs.InsetsLeft));
+		OutDeviceParameters.Add(FAndroidProfileSelectorSourceProperties::SRC_InsetsTop, FString::Printf(TEXT("%f"), DeviceSpecs.InsetsTop));
+		OutDeviceParameters.Add(FAndroidProfileSelectorSourceProperties::SRC_InsetsRight, FString::Printf(TEXT("%f"), DeviceSpecs.InsetsRight));
+		OutDeviceParameters.Add(FAndroidProfileSelectorSourceProperties::SRC_InsetsBottom, FString::Printf(TEXT("%f"), DeviceSpecs.InsetsBottom));
 	}
 }
 #endif
@@ -174,5 +163,36 @@ bool FAndroidDeviceProfileSelectorModule::GetSelectorPropertyValue(const FName&
 
 void FAndroidDeviceProfileSelectorModule::SetSelectorProperties(const TMap<FName, FString>& SelectorPropertiesIn)
 {
+	const int32 ResolutionX = FCString::Atoi(*SelectorPropertiesIn.FindChecked(FAndroidProfileSelectorSourceProperties::SRC_ResolutionX));
+	const int32 ResolutionY = FCString::Atoi(*SelectorPropertiesIn.FindChecked(FAndroidProfileSelectorSourceProperties::SRC_ResolutionY));
+	const float InsetsLeft = FCString::Atof(*SelectorPropertiesIn.FindChecked(FAndroidProfileSelectorSourceProperties::SRC_InsetsLeft));
+	const float InsetsTop = FCString::Atof(*SelectorPropertiesIn.FindChecked(FAndroidProfileSelectorSourceProperties::SRC_InsetsTop));
+	const float InsetsRight = FCString::Atof(*SelectorPropertiesIn.FindChecked(FAndroidProfileSelectorSourceProperties::SRC_InsetsRight));
+	const float InsetsBottom = FCString::Atof(*SelectorPropertiesIn.FindChecked(FAndroidProfileSelectorSourceProperties::SRC_InsetsBottom));
+
+	FString Orientation;
+	bool bNeedPortrait = false;
+	GConfig->GetString(TEXT("/Script/AndroidRuntimeSettings.AndroidRuntimeSettings"), TEXT("Orientation"), Orientation, GEngineIni);
+	if (Orientation.ToLower().Equals("portrait") || Orientation.ToLower().Equals("reverseportrait") || Orientation.ToLower().Equals("sensorportrait"))
+	{
+		bNeedPortrait = true;
+	}
+
+	if (bNeedPortrait)
+	{
+		SafeZones.X = InsetsLeft;
+		SafeZones.Y = InsetsTop;
+		SafeZones.Z = InsetsRight;
+		SafeZones.W = InsetsBottom;
+		ConstrainedAspectRatio = float(ResolutionX) / float(ResolutionY);
+	}
+	else
+	{
+		SafeZones.X = InsetsTop;
+		SafeZones.Y = InsetsRight;
+		SafeZones.Z = InsetsBottom;
+		SafeZones.W = InsetsLeft;
+		ConstrainedAspectRatio = float(ResolutionY) / float(ResolutionX);
+	}
 	FAndroidDeviceProfileSelector::SetSelectorProperties(SelectorPropertiesIn);
 }
diff --git a/Engine/Plugins/Runtime/AndroidDeviceProfileSelector/Source/AndroidDeviceProfileSelector/Public/AndroidDeviceProfileSelector.h b/Engine/Plugins/Runtime/AndroidDeviceProfileSelector/Source/AndroidDeviceProfileSelector/Public/AndroidDeviceProfileSelector.h
--- a/Engine/Plugins/Runtime/AndroidDeviceProfileSelector/Source/AndroidDeviceProfileSelector/Public/AndroidDeviceProfileSelector.h
+++ b/Engine/Plugins/Runtime/AndroidDeviceProfileSelector/Source/AndroidDeviceProfileSelector/Public/AndroidDeviceProfileSelector.h
@@ -21,6 +21,12 @@ namespace FAndroidProfileSelectorSourceProperties
 	static FName SRC_HMDSystemName(TEXT("SRC_HMDSystemName"));
 	static FName SRC_TotalPhysicalGB(TEXT("SRC_TotalPhysicalGB"));
 	static FName SRC_SM5Available(TEXT("SRC_SM5Available"));
+	static FName SRC_ResolutionX(TEXT("SRC_ResolutionX"));
+	static FName SRC_ResolutionY(TEXT("SRC_ResolutionY"));
+	static FName SRC_InsetsLeft(TEXT("SRC_InsetsLeft"));
+	static FName SRC_InsetsTop(TEXT("SRC_InsetsTop"));
+	static FName SRC_InsetsRight(TEXT("SRC_InsetsRight"));
+	static FName SRC_InsetsBottom(TEXT("SRC_InsetsBottom"));
 };
 
 class ANDROIDDEVICEPROFILESELECTOR_API FAndroidDeviceProfileSelector
diff --git a/Engine/Source/Developer/Android/AndroidDeviceDetection/Private/AndroidDeviceDetectionModule.cpp b/Engine/Source/Developer/Android/AndroidDeviceDetection/Private/AndroidDeviceDetectionModule.cpp
--- a/Engine/Source/Developer/Android/AndroidDeviceDetection/Private/AndroidDeviceDetectionModule.cpp
+++ b/Engine/Source/Developer/Android/AndroidDeviceDetection/Private/AndroidDeviceDetectionModule.cpp
@@ -636,6 +636,27 @@ private:
 					// It doesn't really matter if a mapping already exists. There is no listening local port so no contention between multiple editor instances
 					ExecuteAdbCommand(*DeviceCommand, nullptr, nullptr);
 				}
+
+				FString WindowDisplaysOutput;
+				const FString ExtensionsCommand = FString::Printf(TEXT("-s %s shell dumpsys window displays"), *NewDeviceInfo.SerialNumber);
+				if (!ExecuteAdbCommand(*ExtensionsCommand, &WindowDisplaysOutput, nullptr))
+				{
+					continue;
+				}
+				else
+				{
+					WindowDisplaysOutput = WindowDisplaysOutput.ToLower();
+					const FRegexPattern RegexPattern(TEXT(".*initcutout.*insets=rect\\((\\d+)\\W*,\\s*(\\d+)\\W*-\\s*(\\d+)\\W*,\\s*(\\d+)\\W*\\)"));
+					FRegexMatcher RegexMatcher(RegexPattern, *WindowDisplaysOutput);
+					if (RegexMatcher.FindNext())
+					{
+						// store the insets independently from the resolution.
+						NewDeviceInfo.InsetsLeft = FCString::Atof(*RegexMatcher.GetCaptureGroup(1)) / NewDeviceInfo.ResolutionX;
+						NewDeviceInfo.InsetsTop = FCString::Atof(*RegexMatcher.GetCaptureGroup(2)) / NewDeviceInfo.ResolutionY;
+						NewDeviceInfo.InsetsRight = FCString::Atof(*RegexMatcher.GetCaptureGroup(3)) / NewDeviceInfo.ResolutionX;
+						NewDeviceInfo.InsetsBottom = FCString::Atof(*RegexMatcher.GetCaptureGroup(4)) / NewDeviceInfo.ResolutionY;
+					}
+				}
 			}
 
 			// add the device to the map
@@ -839,6 +860,10 @@ public:
 			DeviceSpecs.DevicePlatform = EPIEPreviewDeviceType::Android;
 			DeviceSpecs.ResolutionX = DeviceInfo->ResolutionX;
 			DeviceSpecs.ResolutionY = DeviceInfo->ResolutionY;
+			DeviceSpecs.InsetsLeft = DeviceInfo->InsetsLeft;
+			DeviceSpecs.InsetsTop = DeviceInfo->InsetsTop;
+			DeviceSpecs.InsetsRight = DeviceInfo->InsetsRight;
+			DeviceSpecs.InsetsBottom = DeviceInfo->InsetsBottom;
 			DeviceSpecs.ResolutionYImmersiveMode = 0;
 			DeviceSpecs.PPI = DeviceInfo->DeviceDPI;
 			DeviceSpecs.ScaleFactors = { 0.25f, 0.5f, 0.75f, 1.0f };
diff --git a/Engine/Source/Developer/Android/AndroidDeviceDetection/Public/Interfaces/IAndroidDeviceDetection.h b/Engine/Source/Developer/Android/AndroidDeviceDetection/Public/Interfaces/IAndroidDeviceDetection.h
--- a/Engine/Source/Developer/Android/AndroidDeviceDetection/Public/Interfaces/IAndroidDeviceDetection.h
+++ b/Engine/Source/Developer/Android/AndroidDeviceDetection/Public/Interfaces/IAndroidDeviceDetection.h
@@ -43,6 +43,12 @@ struct FAndroidDeviceInfo
 	// Holds pixel per inch value.
 	int32 DeviceDPI = 0;
 
+	// Holds the safe zones for the device
+	float InsetsLeft;
+	float InsetsTop;
+	float InsetsRight;
+	float InsetsBottom;
+
 	// Holds the display resolution for the device
 	int32 ResolutionX = 0;
 	int32 ResolutionY = 0;
diff --git a/Engine/Source/Editor/LevelEditor/Private/LevelEditorActions.cpp b/Engine/Source/Editor/LevelEditor/Private/LevelEditorActions.cpp
--- a/Engine/Source/Editor/LevelEditor/Private/LevelEditorActions.cpp
+++ b/Engine/Source/Editor/LevelEditor/Private/LevelEditorActions.cpp
@@ -600,7 +600,16 @@ void FLevelEditorActionCallbacks::PreviewJson_Clicked(FName PlatformName, FName
 					auto GetPreviewFeatureLevelInfo = [&]()
 						{
 							const ERHIFeatureLevel::Type FeatureLevel = GetMaxSupportedFeatureLevel(ShaderPlatform);
-							return FPreviewPlatformInfo(FeatureLevel, ShaderPlatform, PlatformName, FDataDrivenShaderPlatformInfo::GetShaderFormat(ShaderPlatform), DevileProfile->GetFName(), true, PreviewShaderPlatformName, FText::FromName(DevileProfile->GetFName()), ConstrainedAspectRatio);
+							return FPreviewPlatformInfo(
+								FeatureLevel, 
+								ShaderPlatform, PlatformName, 
+								FDataDrivenShaderPlatformInfo::GetShaderFormat(ShaderPlatform), 
+								DevileProfile->GetFName(), 
+								true, 
+								PreviewShaderPlatformName, 
+								FText::FromName(DevileProfile->GetFName()),
+								ConstrainedAspectRatio,
+								DPSelectorModule->GetSafeZones());
 						};
 
 					FPreviewPlatformInfo PreviewFeatureLevelInfo = GetPreviewFeatureLevelInfo();
diff --git a/Engine/Source/Editor/PIEPreviewDeviceSpecification/Public/PIEPreviewDeviceSpecification.h b/Engine/Source/Editor/PIEPreviewDeviceSpecification/Public/PIEPreviewDeviceSpecification.h
--- a/Engine/Source/Editor/PIEPreviewDeviceSpecification/Public/PIEPreviewDeviceSpecification.h
+++ b/Engine/Source/Editor/PIEPreviewDeviceSpecification/Public/PIEPreviewDeviceSpecification.h
@@ -183,6 +183,18 @@ struct FPIEPreviewDeviceSpecifications
 	UPROPERTY()
 	int32 ResolutionYImmersiveMode = 0;
 
+	UPROPERTY()
+	float InsetsLeft = 0.0f;
+	
+	UPROPERTY()
+	float InsetsTop = 0.0f;
+	
+	UPROPERTY()
+	float InsetsRight = 0.0f;
+	
+	UPROPERTY()
+	float InsetsBottom = 0.0f;
+
 	UPROPERTY()
 	int32 PPI = 0;
 
diff --git a/Engine/Source/Editor/UnrealEd/Classes/Editor/EditorEngine.h b/Engine/Source/Editor/UnrealEd/Classes/Editor/EditorEngine.h
--- a/Engine/Source/Editor/UnrealEd/Classes/Editor/EditorEngine.h
+++ b/Engine/Source/Editor/UnrealEd/Classes/Editor/EditorEngine.h
@@ -219,11 +219,22 @@ struct FPreviewPlatformInfo
 	,	PreviewShaderPlatformName(NAME_None)
 	,	PreviewShaderPlatformFriendlyName()
 	,	ConstrainedAspectRatio(0.0f)
+	,	SafeZones(FVector4f(0))
 	{
 		InternalSetFriendlyName();
 	}
 
-	FPreviewPlatformInfo(ERHIFeatureLevel::Type InFeatureLevel, EShaderPlatform InShaderPlatform = EShaderPlatform::SP_NumPlatforms, FName InPreviewPlatformName = NAME_None, FName InPreviewShaderFormatName = NAME_None, FName InDeviceProfileName = NAME_None, bool InbPreviewFeatureLevelActive = false, FName InShaderPlatformName = NAME_None, FText InPreviewShaderPlatformFriendlyName = FText(), float InConstrainedAspectRatio = 0.0f)
+	FPreviewPlatformInfo(ERHIFeatureLevel::Type InFeatureLevel, 
+		EShaderPlatform InShaderPlatform = EShaderPlatform::SP_NumPlatforms, 
+		FName InPreviewPlatformName = NAME_None, 
+		FName InPreviewShaderFormatName = NAME_None, 
+		FName InDeviceProfileName = NAME_None, 
+		bool InbPreviewFeatureLevelActive = false, 
+		FName InShaderPlatformName = NAME_None, 
+		FText InPreviewShaderPlatformFriendlyName = FText(), 
+		float InConstrainedAspectRatio = 0.0f,
+		FVector4f InSafeZones = FVector4f::Zero()
+	)
 	:	PreviewFeatureLevel(InFeatureLevel)
 	,	ShaderPlatform(InShaderPlatform)
 	,	PreviewPlatformName(InPreviewPlatformName)
@@ -233,6 +244,7 @@ struct FPreviewPlatformInfo
 	,	PreviewShaderPlatformName(InShaderPlatformName)
 	,	PreviewShaderPlatformFriendlyName(InPreviewShaderPlatformFriendlyName)
 	,	ConstrainedAspectRatio(InConstrainedAspectRatio)
+	,	SafeZones(InSafeZones)
 	{
 		InternalSetFriendlyName();
 	}
@@ -271,11 +283,17 @@ struct FPreviewPlatformInfo
 	}
 
 	/** Return the Constrained Aspect Ratio of the preview platform */
-	float GetConstrainedAspectRatio()
+	UNREALED_API float GetConstrainedAspectRatio()
 	{
 		return ConstrainedAspectRatio;
 	}
 
+	/** Return the Safe Zones of the preview platform */
+	FVector4f GetSafeZones()
+	{
+		return SafeZones;
+	}
+
 	/** returns the preview feature level if active, or GMaxRHIFeatureLevel otherwise */
 	UNREALED_API ERHIFeatureLevel::Type GetEffectivePreviewFeatureLevel() const;
 	
@@ -290,6 +308,9 @@ private:
 	
 	/** Constrained Aspect Ratio of the preview platform */
 	float ConstrainedAspectRatio;
+	
+	/** Safe Zones of the preview platform */
+	FVector4f SafeZones;
 
 	UNREALED_API void InternalSetFriendlyName();
 };
diff --git a/Engine/Source/Editor/UnrealEd/Private/EditorViewportClient.cpp b/Engine/Source/Editor/UnrealEd/Private/EditorViewportClient.cpp
--- a/Engine/Source/Editor/UnrealEd/Private/EditorViewportClient.cpp
+++ b/Engine/Source/Editor/UnrealEd/Private/EditorViewportClient.cpp
@@ -1025,13 +1025,13 @@ FSceneView* FEditorViewportClient::CalcSceneView(FSceneViewFamily* ViewFamily, c
 	TimeForForceRedraw = 0.0;
 
 	bool bConstrainAspectRatio = bUseControllingActorViewInfo && ControllingActorViewInfo.bConstrainAspectRatio;
-	if (GEditor->IsFeatureLevelPreviewActive() && GEditor->GetAllowConstrainedAspectRatioInPreview())
+	if (GEditor != nullptr && GEditor->IsFeatureLevelPreviewActive() && GEditor->GetAllowConstrainedAspectRatioInPreview())
 	{
-		if (GEditor->PreviewPlatform.GetConstrainedAspectRatio() != 0.0f)
-		{
-			AspectRatio = GEditor->PreviewPlatform.GetConstrainedAspectRatio();
-			bConstrainAspectRatio = true;
-		}
+		FSlateApplication::Get().OnConstrainedAspectRatioChanged.Broadcast(GEditor->PreviewPlatform.GetConstrainedAspectRatio());
+	}
+	else
+	{
+		FSlateApplication::Get().OnConstrainedAspectRatioChanged.Broadcast(0.f);
 	}
 
 	EAspectRatioAxisConstraint AspectRatioAxisConstraint = GetDefault<ULevelEditorViewportSettings>()->AspectRatioAxisConstraint;
diff --git a/Engine/Source/Editor/UnrealEd/Private/SEditorViewport.cpp b/Engine/Source/Editor/UnrealEd/Private/SEditorViewport.cpp
--- a/Engine/Source/Editor/UnrealEd/Private/SEditorViewport.cpp
+++ b/Engine/Source/Editor/UnrealEd/Private/SEditorViewport.cpp
@@ -35,6 +35,8 @@
 #if WITH_DUMPGPU
 	#include "RenderGraph.h"
 #endif
+#include "LevelEditorViewport.h"
+#include "SEditorViewportGridPanel.h"
 
 #define LOCTEXT_NAMESPACE "EditorViewport"
 
@@ -213,13 +215,51 @@ void SEditorViewport::Construct( const FArguments& InArgs )
 	}
 
 	// clang-format off
-	VerticalBox->AddSlot()
-	[
-		SNew(SGlobalPlayWorldActions)
-		[
-			ViewportWidget.ToSharedRef()
-		]
-	];
+	
+	bool bIsMainViewport = false;
+	float MainViewportX = 0;
+	float MainViewportY = 0;
+	TSharedPtr<FEditorViewportClient> EditorViewportClient = GetViewportClient();
+	if (EditorViewportClient.IsValid())
+	{
+		if (EditorViewportClient->IsLevelEditorClient())
+		{
+			const FLevelEditorViewportClient* LevelClient = static_cast<FLevelEditorViewportClient*>(EditorViewportClient.Get());
+			if (LevelClient->IsPerspective())
+			{
+				bIsMainViewport = true;
+			}
+
+			if (LevelClient->Viewport != nullptr)
+			{
+				MainViewportX = LevelClient->Viewport->GetSizeXY().X;
+				MainViewportY = LevelClient->Viewport->GetSizeXY().Y;
+			}
+		}
+	}
+
+	if (bIsMainViewport)
+	{
+		// We wrap the main editor viewport in the middle of a 3x3 grid in order to restrict aspect ratio for preview platforms that request it.
+		VerticalBox->AddSlot()
+			[
+				SNew(SGlobalPlayWorldActions)
+					[
+						SNew(SEditorViewportGridPanel)
+						.ViewportWidget(ViewportWidget)
+					]
+			];
+	}
+	else
+	{
+		VerticalBox->AddSlot()
+			[
+				SNew(SGlobalPlayWorldActions)
+					[
+						ViewportWidget.ToSharedRef()
+					]
+			];
+	}
 	// clang-format on
 
 	
diff --git a/Engine/Source/Editor/UnrealEd/Private/SEditorViewportGridPanel.cpp b/Engine/Source/Editor/UnrealEd/Private/SEditorViewportGridPanel.cpp
new file mode 100644
--- /dev/null
+++ b/Engine/Source/Editor/UnrealEd/Private/SEditorViewportGridPanel.cpp
@@ -0,0 +1,72 @@
+#include "SEditorViewportGridPanel.h"
+#include "Framework/Application/SlateApplication.h"
+
+SLATE_IMPLEMENT_WIDGET(SEditorViewportGridPanel)
+void SEditorViewportGridPanel::PrivateRegisterAttributes(FSlateAttributeInitializer& AttributeInitializer)
+{
+}
+
+void SEditorViewportGridPanel::Construct(const FArguments& InArgs)
+{
+	SGridPanel::FArguments SuperArgs;
+	SuperArgs.FillColumn(0, 0.f);
+	SuperArgs.FillColumn(1, 1.f);
+	SuperArgs.FillColumn(2, 0.f);
+	SuperArgs.FillRow(0, 0.f);
+	SuperArgs.FillRow(1, 1.f);
+	SuperArgs.FillRow(2, 0.f);
+	SuperArgs._Visibility = InArgs._Visibility;
+	checkf(InArgs._ViewportWidget.IsSet(), TEXT("ViewportWidget must be set for SEditorViewportGridPanel widget!"));
+	SuperArgs + SGridPanel::Slot(1, 1)[InArgs._ViewportWidget.Get().ToSharedRef()];
+	SGridPanel::Construct(SuperArgs);
+
+	DebugAspectRatio = 0.f;
+	FSlateApplication::Get().OnConstrainedAspectRatioChanged.AddSP(this, &SEditorViewportGridPanel::UpdateAspectRatio);
+}
+
+void SEditorViewportGridPanel::OnArrangeChildren(const FGeometry& AllottedGeometry, FArrangedChildren& ArrangedChildren) const
+{
+	// To make sense of this, look in SEditorViewport::Construct
+	// We wrap the main editor viewport in the middle of a 3x3 grid in order to restrict aspect ratio for preview platforms that request it.
+	{
+		// Default values, fill the middle cell of the grid.
+		ColFillCoefficients[0].Set(0.f);
+		ColFillCoefficients[1].Set(1.f);
+		ColFillCoefficients[2].Set(0.f);
+		RowFillCoefficients[0].Set(0.f);
+		RowFillCoefficients[1].Set(1.f);
+		RowFillCoefficients[2].Set(0.f);
+
+		if (DebugAspectRatio != 0.f)
+		{
+			const float ContainerX = AllottedGeometry.GetLocalSize().X;
+			const float ContainerY = AllottedGeometry.GetLocalSize().Y;
+			const float NewX = ContainerY * DebugAspectRatio;
+			const float NewY = ContainerX / DebugAspectRatio;
+
+			if (NewX < ContainerX)
+			{
+				const float Offset = (ContainerX - NewX) / 2;
+				const float Center = (ContainerX / Offset) - 2;
+				ColFillCoefficients[0].Set(1.f);
+				ColFillCoefficients[1].Set(Center);
+				ColFillCoefficients[2].Set(1.f);
+			}
+
+			if (NewY < ContainerY)
+			{
+				const float Offset = (ContainerY - NewY) / 2;
+				const float Center = (ContainerY / Offset) - 2;
+				RowFillCoefficients[0].Set(1.f);
+				RowFillCoefficients[1].Set(Center);
+				RowFillCoefficients[2].Set(1.f);
+			}
+		}
+	}
+	SGridPanel::OnArrangeChildren(AllottedGeometry, ArrangedChildren);
+}
+
+void SEditorViewportGridPanel::UpdateAspectRatio(const float& AspectRatio)
+{
+	DebugAspectRatio = AspectRatio;
+}
\ No newline at end of file
diff --git a/Engine/Source/Editor/UnrealEd/Private/SEditorViewportGridPanel.h b/Engine/Source/Editor/UnrealEd/Private/SEditorViewportGridPanel.h
new file mode 100644
--- /dev/null
+++ b/Engine/Source/Editor/UnrealEd/Private/SEditorViewportGridPanel.h
@@ -0,0 +1,23 @@
+#pragma once
+
+#include "Widgets/Layout/SGridPanel.h"
+#include "Widgets/SWidget.h"
+
+class SEditorViewportGridPanel : public SGridPanel
+{
+	SLATE_DECLARE_WIDGET_API(SEditorViewportGridPanel, SGridPanel, UNREALED_API)
+
+public:
+	SLATE_BEGIN_ARGS(SEditorViewportGridPanel)
+		{}
+		SLATE_ATTRIBUTE(TSharedPtr<SWidget>, ViewportWidget);
+	SLATE_END_ARGS()
+
+	UNREALED_API void Construct(const FArguments& InArgs);
+	UNREALED_API virtual void OnArrangeChildren(const FGeometry& AllottedGeometry, FArrangedChildren& ArrangedChildren) const override;
+	UNREALED_API void UpdateAspectRatio(const float& AspectRatio);
+
+private:
+	TAttribute<TSharedPtr<SWidget>> ViewportWidget;
+	float DebugAspectRatio;
+};
\ No newline at end of file
diff --git a/Engine/Source/Runtime/Engine/Private/LocalPlayer.cpp b/Engine/Source/Runtime/Engine/Private/LocalPlayer.cpp
--- a/Engine/Source/Runtime/Engine/Private/LocalPlayer.cpp
+++ b/Engine/Source/Runtime/Engine/Private/LocalPlayer.cpp
@@ -1222,17 +1222,9 @@ bool ULocalPlayer::GetProjectionData(FViewport* Viewport, FSceneViewProjectionDa
 	// Get the viewpoint.
 	FMinimalViewInfo ViewInfo;
 	GetViewPoint(/*out*/ ViewInfo);
+	
+
 
-#if WITH_EDITOR
-	if (GEditor != nullptr && GEditor->IsFeatureLevelPreviewActive() && GEditor->GetAllowConstrainedAspectRatioInPreview())
-	{
-		if (GEditor->PreviewPlatform.GetConstrainedAspectRatio() != 0.0f)
-		{
-			ViewInfo.AspectRatio = GEditor->PreviewPlatform.GetConstrainedAspectRatio();
-			ViewInfo.bConstrainAspectRatio = true;
-		}
-	}
-#endif
 
 	// If stereo rendering is enabled, update the size and offset appropriately for this pass
 	const bool bNeedStereo = StereoViewIndex != INDEX_NONE && GEngine->IsStereoscopic3D();
@@ -1279,6 +1271,24 @@ bool ULocalPlayer::GetProjectionData(FViewport* Viewport, FSceneViewProjectionDa
 		FPlane(0,	1,	0,	0),
 		FPlane(0,	0,	0,	1));
 
+#if WITH_EDITOR
+	if (GEditor != nullptr && GEditor->IsFeatureLevelPreviewActive() && GEditor->GetAllowConstrainedAspectRatioInPreview())
+	{
+		FSlateApplication::Get().OnConstrainedAspectRatioChanged.Broadcast(GEditor->PreviewPlatform.GetConstrainedAspectRatio());
+
+		FVector4f Safezones = FVector4f::Zero();
+		Safezones = GEditor->PreviewPlatform.GetSafeZones();
+		FSlateApplication::Get().SetCustomSafeZone(Safezones);
+		FSlateApplication::Get().OnDebugSafeZoneChanged.Broadcast(Safezones, true);
+	}
+	else
+	{
+		FSlateApplication::Get().OnConstrainedAspectRatioChanged.Broadcast(0.f);
+		FSlateApplication::Get().ResetCustomSafeZone();
+		FSlateApplication::Get().OnDebugSafeZoneChanged.Broadcast(FMargin(), true);
+	}
+#endif
+
 	// @todo viewext this use case needs to be revisited
 	if (!bNeedStereo)
 	{
diff --git a/Engine/Source/Runtime/Engine/Public/IDeviceProfileSelectorModule.h b/Engine/Source/Runtime/Engine/Public/IDeviceProfileSelectorModule.h
--- a/Engine/Source/Runtime/Engine/Public/IDeviceProfileSelectorModule.h
+++ b/Engine/Source/Runtime/Engine/Public/IDeviceProfileSelectorModule.h
@@ -19,7 +19,10 @@ class IDeviceProfileSelectorModule
 	: public IModuleInterface
 {
 public:
-	IDeviceProfileSelectorModule() :ConstrainedAspectRatio(0.0f) {}
+	IDeviceProfileSelectorModule() :
+		ConstrainedAspectRatio(0.0f), 
+		SafeZones(FVector4f::Zero())
+	{}
 	
 	/**
 	 * Run the logic to choose an appropriate device profile for this session
@@ -54,6 +57,12 @@ public:
 	 */
 	virtual float GetConstrainedAspectRatio() { return ConstrainedAspectRatio; }
 	
+	/**
+	 * Get the Safe Zones of the Device
+	 *  @return Safe Zones of the Device
+	 */
+	virtual FVector4f GetSafeZones() { return SafeZones; }
+	
 	/**
 	 * Can we get device parameters from Json
 	 *  @return Wheter we can get device parameters from Json
@@ -91,4 +100,5 @@ public:
 
 protected:
 	float ConstrainedAspectRatio;
+	FVector4f SafeZones;
 };
diff --git a/Engine/Source/Runtime/Slate/Public/Widgets/Layout/SGridPanel.h b/Engine/Source/Runtime/Slate/Public/Widgets/Layout/SGridPanel.h
--- a/Engine/Source/Runtime/Slate/Public/Widgets/Layout/SGridPanel.h
+++ b/Engine/Source/Runtime/Slate/Public/Widgets/Layout/SGridPanel.h
@@ -333,9 +333,10 @@ private:
 	/** Total desires size along each axis. */
 	FVector2D TotalDesiredSizes;
 
+protected:
 	/** Fill coefficients for the columns */
-	TArray<TAttribute<float>> ColFillCoefficients;
+	mutable TArray<TAttribute<float>> ColFillCoefficients;
 
 	/** Fill coefficients for the rows */
-	TArray<TAttribute<float>> RowFillCoefficients;
-};
+	mutable TArray<TAttribute<float>> RowFillCoefficients;
+};
\ No newline at end of file
diff --git a/Engine/Source/Runtime/SlateCore/Public/Application/SlateApplicationBase.h b/Engine/Source/Runtime/SlateCore/Public/Application/SlateApplicationBase.h
--- a/Engine/Source/Runtime/SlateCore/Public/Application/SlateApplicationBase.h
+++ b/Engine/Source/Runtime/SlateCore/Public/Application/SlateApplicationBase.h
@@ -13,6 +13,7 @@
 #include "Types/SlateVector2.h"
 
 DECLARE_MULTICAST_DELEGATE_TwoParams(FOnDebugSafeZoneChanged, const FMargin&, bool);
+DECLARE_MULTICAST_DELEGATE_OneParam(FOnConstrainedAspectRatioChanged, const float&);
 
 
 class FActiveTimerHandle;
@@ -666,7 +667,9 @@ public:
 #if WITH_EDITORONLY_DATA
 	FOnDebugSafeZoneChanged OnDebugSafeZoneChanged;
 #endif
-
+#if WITH_EDITOR
+	FOnConstrainedAspectRatioChanged OnConstrainedAspectRatioChanged;
+#endif
 protected:
 #if WITH_ACCESSIBILITY
 	/** Manager for widgets and application to interact with accessibility API */
