Added support to be able to run the emulator, and added a runtime flag to enable it

#android

diff --git a/Engine/Config/BaseEngine.ini b/Engine/Config/BaseEngine.ini
--- a/Engine/Config/BaseEngine.ini
+++ b/Engine/Config/BaseEngine.ini
@@ -3248,6 +3248,7 @@ bBuildForX8664=false
 bBuildForES31=true
 bSupportsVulkan=true
 bSupportsVulkanSM5=false
+bForceEmulatorProfileSelectionInNonShippingBuilds=true
 bDetectVulkanByDefault=true
 bSplitIntoSeparateApks=false
 bPackageDataInsideApk=false
diff --git a/Engine/Plugins/Runtime/AndroidDeviceProfileSelector/Source/AndroidDeviceProfileSelectorRuntime/Private/AndroidDeviceProfileSelectorRuntimeModule.cpp b/Engine/Plugins/Runtime/AndroidDeviceProfileSelector/Source/AndroidDeviceProfileSelectorRuntime/Private/AndroidDeviceProfileSelectorRuntimeModule.cpp
--- a/Engine/Plugins/Runtime/AndroidDeviceProfileSelector/Source/AndroidDeviceProfileSelectorRuntime/Private/AndroidDeviceProfileSelectorRuntimeModule.cpp
+++ b/Engine/Plugins/Runtime/AndroidDeviceProfileSelector/Source/AndroidDeviceProfileSelectorRuntime/Private/AndroidDeviceProfileSelectorRuntimeModule.cpp
@@ -121,7 +121,31 @@ FString const FAndroidDeviceProfileSelectorRuntimeModule::GetRuntimeDeviceProfil
 		const FString& DeviceModel = AndroidParams.FindChecked(FAndroidProfileSelectorSourceProperties::SRC_DeviceModel);
 
 		CheckForJavaSurfaceViewWorkaround(DeviceMake, DeviceModel);
+#if !UE_BUILD_SHIPPING
 
+		{
+			bool bForceEmulatorProfileSelectionInNonShippingBuilds = true;
+			GConfig->GetBool(TEXT("/Script/AndroidRuntimeSettings.AndroidRuntimeSettings"), TEXT("bForceEmulatorProfileSelectionInNonShippingBuilds"), bForceEmulatorProfileSelectionInNonShippingBuilds, GEngineIni);
+			
+			if (bForceEmulatorProfileSelectionInNonShippingBuilds && DeviceMake == TEXT("Google"))
+			{
+				if (DeviceModel == TEXT("HPE device"))
+				{
+					ProfileName = TEXT("Android_PC_Emulator");
+					UE_LOG(LogAndroid, Log, TEXT("Forced Profile: [Android_PC_Emulator]"));
+					return ProfileName;
+				}
+				
+				if (AndroidParams.FindChecked(FAndroidProfileSelectorSourceProperties::SRC_GPUFamily).StartsWith(TEXT("Android Emulator")))
+				{
+					ProfileName = TEXT("Android_Emulator");
+					UE_LOG(LogAndroid, Log, TEXT("Forced Profile: [Android_Emulator]"));
+					return ProfileName;
+				}
+			}
+		}
+#endif
+		
 		// Use override from ConfigRules if set
 		const FString* ConfigProfile = FPlatformMisc::GetConfigRulesVariable(TEXT("Profile"));
 		if (ConfigProfile != nullptr)
diff --git a/Engine/Source/Runtime/Android/AndroidRuntimeSettings/Classes/AndroidRuntimeSettings.h b/Engine/Source/Runtime/Android/AndroidRuntimeSettings/Classes/AndroidRuntimeSettings.h
--- a/Engine/Source/Runtime/Android/AndroidRuntimeSettings/Classes/AndroidRuntimeSettings.h
+++ b/Engine/Source/Runtime/Android/AndroidRuntimeSettings/Classes/AndroidRuntimeSettings.h
@@ -463,6 +463,9 @@ public:
 	ToolTip = "Whether to render to an offscreen surface instead of render to backbuffer directly on android opengl platform. Enable it if you'd like to support UMG background blur on android opengl."))
 	bool bAndroidOpenGLSupportsBackbufferSampling;
 
+	UPROPERTY(GlobalConfig, EditAnywhere, Category = Build, meta = (DisplayName = "Force Emulator profile selection in non-shipping builds"))
+	bool bForceEmulatorProfileSelectionInNonShippingBuilds;
+	
 	// Whether to detect Vulkan device support by default, if the project is packaged with Vulkan support. If unchecked, the -detectvulkan commandline will enable Vulkan detection.
 	UPROPERTY(GlobalConfig, EditAnywhere, AdvancedDisplay, Category = Build, meta = (DisplayName = "Detect Vulkan device support", EditCondition = "bSupportsVulkan"))
 	bool bDetectVulkanByDefault;
