Add fallback to AndroidDeviceProfileSelectorModule
Add warning log if the Preview via Json doesn't find a DP from the MatchingProfile Rules

diff --git a/Engine/Plugins/Runtime/AndroidDeviceProfileSelector/Source/AndroidDeviceProfileSelector/Private/AndroidDeviceProfileSelectorModule.cpp b/Engine/Plugins/Runtime/AndroidDeviceProfileSelector/Source/AndroidDeviceProfileSelector/Private/AndroidDeviceProfileSelectorModule.cpp
--- a/Engine/Plugins/Runtime/AndroidDeviceProfileSelector/Source/AndroidDeviceProfileSelector/Private/AndroidDeviceProfileSelectorModule.cpp
+++ b/Engine/Plugins/Runtime/AndroidDeviceProfileSelector/Source/AndroidDeviceProfileSelector/Private/AndroidDeviceProfileSelectorModule.cpp
@@ -15,6 +15,7 @@
 
 #endif
 #include "Misc/ConfigCacheIni.h"
+#include "Android/AndroidPlatformProperties.h"
 
 IMPLEMENT_MODULE(FAndroidDeviceProfileSelectorModule, AndroidDeviceProfileSelector);
 
@@ -123,7 +124,7 @@ void FAndroidDeviceProfileSelectorModule::GetDeviceParametersFromJson(FString& J
 
 const FString FAndroidDeviceProfileSelectorModule::GetDeviceProfileName()
 {
-	FString ProfileName; 
+	FString ProfileName = ANDROID_DEFAULT_DEVICE_PROFILE_NAME;
 
 	// ensure SelectorProperties does actually contain our parameters
 	check(FAndroidDeviceProfileSelector::GetSelectorProperties().Num() > 0);
diff --git a/Engine/Source/Editor/LevelEditor/Private/LevelEditorActions.cpp b/Engine/Source/Editor/LevelEditor/Private/LevelEditorActions.cpp
--- a/Engine/Source/Editor/LevelEditor/Private/LevelEditorActions.cpp
+++ b/Engine/Source/Editor/LevelEditor/Private/LevelEditorActions.cpp
@@ -615,6 +615,10 @@ void FLevelEditorActionCallbacks::PreviewJson_Clicked(FName PlatformName, FName
 					FPreviewPlatformInfo PreviewFeatureLevelInfo = GetPreviewFeatureLevelInfo();
 					GEditor->SetPreviewPlatform(PreviewFeatureLevelInfo, false);
 				}
+				else
+				{
+					UE_LOG(LevelEditorActions, Warning, TEXT("JsonFile %s did not match any MatchProfile"), *JsonFile);
+				}
 			}
 		}
 	}
diff --git a/Engine/Source/Runtime/Core/Private/Android/AndroidPlatformMisc.cpp b/Engine/Source/Runtime/Core/Private/Android/AndroidPlatformMisc.cpp
--- a/Engine/Source/Runtime/Core/Private/Android/AndroidPlatformMisc.cpp
+++ b/Engine/Source/Runtime/Core/Private/Android/AndroidPlatformMisc.cpp
@@ -60,6 +60,7 @@
 #include "VulkanCommon.h"
 #include "IVulkanDynamicRHI.h"
 #include <vulkan/vulkan_android.h>
+#include "Android/AndroidPlatformProperties.h"
 
 #if USE_ANDROID_JNI
 extern AAssetManager * AndroidThunkCpp_GetAssetManager();
@@ -1647,6 +1648,11 @@ bool FAndroidMisc::SupportsTouchInput()
 	return true;
 }
 
+const TCHAR* FAndroidMisc::GetDefaultDeviceProfileName()
+{ 
+	return ANDROID_DEFAULT_DEVICE_PROFILE_NAME; 
+}
+
 extern void AndroidThunkCpp_RegisterForRemoteNotifications();
 extern void AndroidThunkCpp_UnregisterForRemoteNotifications();
 extern bool AndroidThunkCpp_IsAllowedRemoteNotifications();
diff --git a/Engine/Source/Runtime/Core/Public/Android/AndroidPlatformMisc.h b/Engine/Source/Runtime/Core/Public/Android/AndroidPlatformMisc.h
--- a/Engine/Source/Runtime/Core/Public/Android/AndroidPlatformMisc.h
+++ b/Engine/Source/Runtime/Core/Public/Android/AndroidPlatformMisc.h
@@ -106,7 +106,7 @@ public:
 	static CORE_API void GetValidTargetPlatforms(class TArray<class FString>& TargetPlatformNames);
 	static CORE_API bool GetUseVirtualJoysticks();
 	static CORE_API bool SupportsTouchInput();
-	static const TCHAR* GetDefaultDeviceProfileName() { return TEXT("Android_Default"); }
+	static const TCHAR* GetDefaultDeviceProfileName();
 	static CORE_API bool GetVolumeButtonsHandledBySystem();
 	static CORE_API void SetVolumeButtonsHandledBySystem(bool enabled);
 	// Returns current volume, 0-15
diff --git a/Engine/Source/Runtime/Core/Public/Android/AndroidPlatformProperties.h b/Engine/Source/Runtime/Core/Public/Android/AndroidPlatformProperties.h
--- a/Engine/Source/Runtime/Core/Public/Android/AndroidPlatformProperties.h
+++ b/Engine/Source/Runtime/Core/Public/Android/AndroidPlatformProperties.h
@@ -12,6 +12,7 @@
 #include "CoreTypes.h"
 #include "GenericPlatform/GenericPlatformProperties.h"
 
+#define ANDROID_DEFAULT_DEVICE_PROFILE_NAME TEXT("Android_Default")
 
 /**
  * Implements Android platform properties.
