Android Preview - Add warning when no devices were found when attempting to export  device parameters to json

#tests local test with and without android phone & ADB configured

diff --git a/Engine/Plugins/Runtime/AndroidDeviceProfileSelector/Source/AndroidDeviceProfileSelector/Private/AndroidDeviceProfileSelectorModule.cpp b/Engine/Plugins/Runtime/AndroidDeviceProfileSelector/Source/AndroidDeviceProfileSelector/Private/AndroidDeviceProfileSelectorModule.cpp
--- a/Engine/Plugins/Runtime/AndroidDeviceProfileSelector/Source/AndroidDeviceProfileSelector/Private/AndroidDeviceProfileSelectorModule.cpp
+++ b/Engine/Plugins/Runtime/AndroidDeviceProfileSelector/Source/AndroidDeviceProfileSelector/Private/AndroidDeviceProfileSelectorModule.cpp
@@ -58,15 +58,22 @@ void FAndroidDeviceProfileSelectorModule::ExportDeviceParametersToJson(FString&
 		FScopeLock ExportLock(DeviceDetection->GetDeviceMapLock());
 
 		const TMap<FString, FAndroidDeviceInfo>& Devices = DeviceDetection->GetDeviceMap();
-		for (auto DeviceTuple : Devices)
+		if (Devices.Num() == 0)
 		{
-			const FAndroidDeviceInfo& DeviceInfo = DeviceTuple.Value;
-			FString DeviceName = FString::Printf(TEXT("%s_%s(OS%s)"), *DeviceInfo.DeviceBrand, *DeviceInfo.Model, *DeviceInfo.HumanAndroidVersion);
-			if (!AlreadyExported.Find(DeviceName))
+			UE_LOG(LogAndroidDPSelector, Warning, TEXT("ExportDeviceParametersToJson - No device detected. Check that $ANDROID_HOME is set and USB debugging is enabled on the device."));
+		}
+		else
+		{
+			for (auto DeviceTuple : Devices)
 			{
-				FString ExportPath = FolderLocation / (DeviceName + TEXT(".json"));
-				DeviceDetection->ExportDeviceProfile(ExportPath, DeviceTuple.Key);
-				AlreadyExported.Add(DeviceName);
+				const FAndroidDeviceInfo& DeviceInfo = DeviceTuple.Value;
+				FString DeviceName = FString::Printf(TEXT("%s_%s(OS%s)"), *DeviceInfo.DeviceBrand, *DeviceInfo.Model, *DeviceInfo.HumanAndroidVersion);
+				if (!AlreadyExported.Find(DeviceName))
+				{
+					FString ExportPath = FolderLocation / (DeviceName + TEXT(".json"));
+					DeviceDetection->ExportDeviceProfile(ExportPath, DeviceTuple.Key);
+					AlreadyExported.Add(DeviceName);
+				}
 			}
 		}
 	}
