Remove direct access to UHT generated functions when UE_WITH_CONSTINIT_UOBJECT is true.
In a future change, constinit reflection info will be available early in boot for access from device profile modules.

diff --git a/Engine/Plugins/Runtime/AndroidDeviceProfileSelector/Source/AndroidDeviceProfileSelector/Private/AndroidDeviceProfileSelector.cpp b/Engine/Plugins/Runtime/AndroidDeviceProfileSelector/Source/AndroidDeviceProfileSelector/Private/AndroidDeviceProfileSelector.cpp
--- a/Engine/Plugins/Runtime/AndroidDeviceProfileSelector/Source/AndroidDeviceProfileSelector/Private/AndroidDeviceProfileSelector.cpp
+++ b/Engine/Plugins/Runtime/AndroidDeviceProfileSelector/Source/AndroidDeviceProfileSelector/Private/AndroidDeviceProfileSelector.cpp
@@ -22,10 +22,12 @@ UAndroidJavaSurfaceViewDevices::UAndroidJavaSurfaceViewDevices(const FObjectInit
 
 static UAndroidDeviceProfileMatchingRules* GetAndroidDeviceProfileMatchingRules()
 {
+#if !UE_WITH_CONSTINIT_UOBJECT
 	// We need to initialize the class early as device profiles need to be evaluated before ProcessNewlyLoadedUObjects can be called.
 	extern UClass* Z_Construct_UClass_UAndroidDeviceProfileMatchingRules();
 	CreatePackage(UAndroidDeviceProfileMatchingRules::StaticPackage());
 	Z_Construct_UClass_UAndroidDeviceProfileMatchingRules();
+#endif // !UE_WITH_CONSTINIT_UOBJECT
 
 	// Get the default object which will has the values from DeviceProfiles.ini
 	UAndroidDeviceProfileMatchingRules* Rules = Cast<UAndroidDeviceProfileMatchingRules>(UAndroidDeviceProfileMatchingRules::StaticClass()->GetDefaultObject());
diff --git a/Engine/Plugins/Runtime/AndroidDeviceProfileSelector/Source/AndroidDeviceProfileSelectorRuntime/Private/AndroidDeviceProfileSelectorRuntimeModule.cpp b/Engine/Plugins/Runtime/AndroidDeviceProfileSelector/Source/AndroidDeviceProfileSelectorRuntime/Private/AndroidDeviceProfileSelectorRuntimeModule.cpp
--- a/Engine/Plugins/Runtime/AndroidDeviceProfileSelector/Source/AndroidDeviceProfileSelectorRuntime/Private/AndroidDeviceProfileSelectorRuntimeModule.cpp
+++ b/Engine/Plugins/Runtime/AndroidDeviceProfileSelector/Source/AndroidDeviceProfileSelectorRuntime/Private/AndroidDeviceProfileSelectorRuntimeModule.cpp
@@ -188,8 +188,10 @@ void FAndroidDeviceProfileSelectorRuntimeModule::CheckForJavaSurfaceViewWorkarou
 {
 #if USE_ANDROID_JNI
 	// We need to initialize the class early as device profiles need to be evaluated before ProcessNewlyLoadedUObjects can be called.
+#if !UE_WITH_CONSTINIT_UOBJECT
 	extern UClass* Z_Construct_UClass_UAndroidJavaSurfaceViewDevices();
 	Z_Construct_UClass_UAndroidJavaSurfaceViewDevices();
+#endif // !UE_WITH_CONSTINIT_UOBJECT
 
 	const UAndroidJavaSurfaceViewDevices *const SurfaceViewDevices = Cast<UAndroidJavaSurfaceViewDevices>(UAndroidJavaSurfaceViewDevices::StaticClass()->GetDefaultObject());
 	check(SurfaceViewDevices);
