Add ConfigRules to Android Preview Json
Clean-up unecessary usage of FName in AndroidDeviceSelector

#Synced-CL 50766215

diff --git a/Engine/Plugins/Runtime/AndroidDeviceProfileSelector/Source/AndroidDeviceProfileSelector/Private/AndroidDeviceProfileSelector.cpp b/Engine/Plugins/Runtime/AndroidDeviceProfileSelector/Source/AndroidDeviceProfileSelector/Private/AndroidDeviceProfileSelector.cpp
--- a/Engine/Plugins/Runtime/AndroidDeviceProfileSelector/Source/AndroidDeviceProfileSelector/Private/AndroidDeviceProfileSelector.cpp
+++ b/Engine/Plugins/Runtime/AndroidDeviceProfileSelector/Source/AndroidDeviceProfileSelector/Private/AndroidDeviceProfileSelector.cpp
@@ -35,47 +35,55 @@ static UAndroidDeviceProfileMatchingRules* GetAndroidDeviceProfileMatchingRules(
 	return Rules;
 }
 
-TMap<FName, FString> FAndroidDeviceProfileSelector::SelectorProperties;
+TMap<FString, FString> FAndroidDeviceProfileSelector::SelectorProperties;
 
 void FAndroidDeviceProfileSelector::VerifySelectorParams()
 {
-	check(SelectorProperties.Find(FAndroidProfileSelectorSourceProperties::SRC_GPUFamily));
-	check(SelectorProperties.Find(FAndroidProfileSelectorSourceProperties::SRC_GLVersion));
-	check(SelectorProperties.Find(FAndroidProfileSelectorSourceProperties::SRC_AndroidVersion));
-	check(SelectorProperties.Find(FAndroidProfileSelectorSourceProperties::SRC_DeviceMake));
-	check(SelectorProperties.Find(FAndroidProfileSelectorSourceProperties::SRC_DeviceModel));
-	check(SelectorProperties.Find(FAndroidProfileSelectorSourceProperties::SRC_DeviceBuildNumber));
-	check(SelectorProperties.Find(FAndroidProfileSelectorSourceProperties::SRC_VulkanVersion));
-	check(SelectorProperties.Find(FAndroidProfileSelectorSourceProperties::SRC_UsingHoudini));
-	check(SelectorProperties.Find(FAndroidProfileSelectorSourceProperties::SRC_VulkanAvailable));
-	check(SelectorProperties.Find(FAndroidProfileSelectorSourceProperties::SRC_Hardware));
-	check(SelectorProperties.Find(FAndroidProfileSelectorSourceProperties::SRC_Chipset));
-	check(SelectorProperties.Find(FAndroidProfileSelectorSourceProperties::SRC_HMDSystemName));
-	check(SelectorProperties.Find(FAndroidProfileSelectorSourceProperties::SRC_TotalPhysicalGB));
-	check(SelectorProperties.Find(FAndroidProfileSelectorSourceProperties::SRC_SM5Available));
-	check(SelectorProperties.Find(FAndroidProfileSelectorSourceProperties::SRC_VKQuality));
+	auto GetSelectorPropertiesValue = [](const FStringView& InStringView)
+		{
+			return SelectorProperties.FindByHash(GetTypeHash(InStringView), InStringView);
+		};
+	check(GetSelectorPropertiesValue(FAndroidProfileSelectorSourceProperties::SRC_GPUFamily));
+	check(GetSelectorPropertiesValue(FAndroidProfileSelectorSourceProperties::SRC_GLVersion));
+	check(GetSelectorPropertiesValue(FAndroidProfileSelectorSourceProperties::SRC_AndroidVersion));
+	check(GetSelectorPropertiesValue(FAndroidProfileSelectorSourceProperties::SRC_DeviceMake));
+	check(GetSelectorPropertiesValue(FAndroidProfileSelectorSourceProperties::SRC_DeviceModel));
+	check(GetSelectorPropertiesValue(FAndroidProfileSelectorSourceProperties::SRC_DeviceBuildNumber));
+	check(GetSelectorPropertiesValue(FAndroidProfileSelectorSourceProperties::SRC_VulkanVersion));
+	check(GetSelectorPropertiesValue(FAndroidProfileSelectorSourceProperties::SRC_UsingHoudini));
+	check(GetSelectorPropertiesValue(FAndroidProfileSelectorSourceProperties::SRC_VulkanAvailable));
+	check(GetSelectorPropertiesValue(FAndroidProfileSelectorSourceProperties::SRC_Hardware));
+	check(GetSelectorPropertiesValue(FAndroidProfileSelectorSourceProperties::SRC_Chipset));
+	check(GetSelectorPropertiesValue(FAndroidProfileSelectorSourceProperties::SRC_HMDSystemName));
+	check(GetSelectorPropertiesValue(FAndroidProfileSelectorSourceProperties::SRC_TotalPhysicalGB));
+	check(GetSelectorPropertiesValue(FAndroidProfileSelectorSourceProperties::SRC_SM5Available));
+	check(GetSelectorPropertiesValue(FAndroidProfileSelectorSourceProperties::SRC_VKQuality));
 }
 
 FString FAndroidDeviceProfileSelector::FindMatchingProfile(const FString& FallbackProfileName)
 {
 	FString OutProfileName = FallbackProfileName;
 	FString CommandLine = FCommandLine::Get();
-	
-	FString GPUFamily			= SelectorProperties.FindChecked(FAndroidProfileSelectorSourceProperties::SRC_GPUFamily);
-	FString GLVersion			= SelectorProperties.FindChecked(FAndroidProfileSelectorSourceProperties::SRC_GLVersion);
-	FString AndroidVersion		= SelectorProperties.FindChecked(FAndroidProfileSelectorSourceProperties::SRC_AndroidVersion);
-	FString DeviceMake			= SelectorProperties.FindChecked(FAndroidProfileSelectorSourceProperties::SRC_DeviceMake);
-	FString DeviceModel			= SelectorProperties.FindChecked(FAndroidProfileSelectorSourceProperties::SRC_DeviceModel);
-	FString DeviceBuildNumber	= SelectorProperties.FindChecked(FAndroidProfileSelectorSourceProperties::SRC_DeviceBuildNumber);
-	FString VulkanVersion		= SelectorProperties.FindChecked(FAndroidProfileSelectorSourceProperties::SRC_VulkanVersion);
-	FString UsingHoudini		= SelectorProperties.FindChecked(FAndroidProfileSelectorSourceProperties::SRC_UsingHoudini);
-	FString VulkanAvailable		= SelectorProperties.FindChecked(FAndroidProfileSelectorSourceProperties::SRC_VulkanAvailable);
-	FString Hardware			= SelectorProperties.FindChecked(FAndroidProfileSelectorSourceProperties::SRC_Hardware);
-	FString Chipset				= SelectorProperties.FindChecked(FAndroidProfileSelectorSourceProperties::SRC_Chipset);
-	FString TotalPhysicalGB		= SelectorProperties.FindChecked(FAndroidProfileSelectorSourceProperties::SRC_TotalPhysicalGB);
-	FString HMDSystemName		= SelectorProperties.FindChecked(FAndroidProfileSelectorSourceProperties::SRC_HMDSystemName);
-	FString SM5Available		= SelectorProperties.FindChecked(FAndroidProfileSelectorSourceProperties::SRC_SM5Available);
-	FString VKQuality			= SelectorProperties.FindChecked(FAndroidProfileSelectorSourceProperties::SRC_VKQuality);
+
+	auto GetSelectorPropertiesValueChecked = [](const FStringView& InStringView)
+		{
+			return SelectorProperties.FindByHashChecked(GetTypeHash(InStringView), InStringView);
+		};
+	FString GPUFamily = GetSelectorPropertiesValueChecked(FAndroidProfileSelectorSourceProperties::SRC_GPUFamily);
+	FString GLVersion = GetSelectorPropertiesValueChecked(FAndroidProfileSelectorSourceProperties::SRC_GLVersion);
+	FString AndroidVersion = GetSelectorPropertiesValueChecked(FAndroidProfileSelectorSourceProperties::SRC_AndroidVersion);
+	FString DeviceMake = GetSelectorPropertiesValueChecked(FAndroidProfileSelectorSourceProperties::SRC_DeviceMake);
+	FString DeviceModel = GetSelectorPropertiesValueChecked(FAndroidProfileSelectorSourceProperties::SRC_DeviceModel);
+	FString DeviceBuildNumber = GetSelectorPropertiesValueChecked(FAndroidProfileSelectorSourceProperties::SRC_DeviceBuildNumber);
+	FString VulkanVersion = GetSelectorPropertiesValueChecked(FAndroidProfileSelectorSourceProperties::SRC_VulkanVersion);
+	FString UsingHoudini = GetSelectorPropertiesValueChecked(FAndroidProfileSelectorSourceProperties::SRC_UsingHoudini);
+	FString VulkanAvailable = GetSelectorPropertiesValueChecked(FAndroidProfileSelectorSourceProperties::SRC_VulkanAvailable);
+	FString Hardware = GetSelectorPropertiesValueChecked(FAndroidProfileSelectorSourceProperties::SRC_Hardware);
+	FString Chipset = GetSelectorPropertiesValueChecked(FAndroidProfileSelectorSourceProperties::SRC_Chipset);
+	FString TotalPhysicalGB = GetSelectorPropertiesValueChecked(FAndroidProfileSelectorSourceProperties::SRC_TotalPhysicalGB);
+	FString HMDSystemName = GetSelectorPropertiesValueChecked(FAndroidProfileSelectorSourceProperties::SRC_HMDSystemName);
+	FString SM5Available = GetSelectorPropertiesValueChecked(FAndroidProfileSelectorSourceProperties::SRC_SM5Available);
+	FString VKQuality = GetSelectorPropertiesValueChecked(FAndroidProfileSelectorSourceProperties::SRC_VKQuality);
 
 	for (const FProfileMatch& Profile : GetAndroidDeviceProfileMatchingRules()->MatchProfile)
 	{
@@ -148,12 +156,14 @@ FString FAndroidDeviceProfileSelector::FindMatchingProfile(const FString& Fallba
 					{
 						MatchString = VariableValueMatchString;
 #if PLATFORM_ANDROID
-						// TODO: Do this when running device selection with editor builds.
-						if (const FString* ConfigRuleVar = FPlatformMisc::GetConfigRulesVariable(ConfigRuleVarName))
+						const FString* ConfigRuleVar = FPlatformMisc::GetConfigRulesVariable(ConfigRuleVarName);
+#elif WITH_EDITOR
+						const FString* ConfigRuleVar = SelectorProperties.Find(FString::Printf(TEXT("SRC_ConfigRuleVar[%s]"), *ConfigRuleVarName));
+#endif
+						if (ConfigRuleVar)
 						{
 							ConfigRuleString = *ConfigRuleVar;
 						}
-#endif
 					}					
 					break;
 				}
diff --git a/Engine/Plugins/Runtime/AndroidDeviceProfileSelector/Source/AndroidDeviceProfileSelector/Private/AndroidDeviceProfileSelectorModule.cpp b/Engine/Plugins/Runtime/AndroidDeviceProfileSelector/Source/AndroidDeviceProfileSelector/Private/AndroidDeviceProfileSelectorModule.cpp
--- a/Engine/Plugins/Runtime/AndroidDeviceProfileSelector/Source/AndroidDeviceProfileSelector/Private/AndroidDeviceProfileSelectorModule.cpp
+++ b/Engine/Plugins/Runtime/AndroidDeviceProfileSelector/Source/AndroidDeviceProfileSelector/Private/AndroidDeviceProfileSelectorModule.cpp
@@ -4,7 +4,7 @@
 #include "AndroidDeviceProfileSelector.h"
 #include "Modules/ModuleManager.h"
 #if WITH_EDITOR
-
+#include "GenericPlatform/GenericConfigRules.h"
 #if WITH_ANDROID_DEVICE_DETECTION
 #include "Interfaces/IAndroidDeviceDetection.h"
 #include "Interfaces/IAndroidDeviceDetectionModule.h"
@@ -97,7 +97,7 @@ bool FAndroidDeviceProfileSelectorModule::CanExportDeviceParametersToJson()
 #endif
 }
 
-void FAndroidDeviceProfileSelectorModule::GetDeviceParametersFromJson(FString& JsonLocation, TMap<FName, FString>& OutDeviceParameters)
+void FAndroidDeviceProfileSelectorModule::GetDeviceParametersFromJson(FString& JsonLocation, TMap<FString, FString>& OutDeviceParameters)
 {
 	TSharedPtr<FJsonObject> JsonRootObject;
 	FString Json;
@@ -112,27 +112,43 @@ void FAndroidDeviceProfileSelectorModule::GetDeviceParametersFromJson(FString& J
 	if (JsonRootObject.IsValid() && FJsonObjectConverter::JsonAttributesToUStruct(JsonRootObject->Values, FPIEPreviewDeviceSpecifications::StaticStruct(), &DeviceSpecs, 0, 0))
 	{
 		FPIEAndroidDeviceProperties& AndroidProperties = DeviceSpecs.AndroidProperties;
-		OutDeviceParameters.Add(FName(TEXT("SRC_GPUFamily")), AndroidProperties.GPUFamily);
-		OutDeviceParameters.Add(FName(TEXT("SRC_GLVersion")), AndroidProperties.GLVersion);
-		OutDeviceParameters.Add(FName(TEXT("SRC_VulkanAvailable")), AndroidProperties.VulkanAvailable ? "true" : "false");
-		OutDeviceParameters.Add(FName(TEXT("SRC_VulkanVersion")), AndroidProperties.VulkanVersion);
-		OutDeviceParameters.Add(FName(TEXT("SRC_AndroidVersion")), AndroidProperties.AndroidVersion);
-		OutDeviceParameters.Add(FName(TEXT("SRC_DeviceMake")), AndroidProperties.DeviceMake);
-		OutDeviceParameters.Add(FName(TEXT("SRC_DeviceModel")), AndroidProperties.DeviceModel);
-		OutDeviceParameters.Add(FName(TEXT("SRC_DeviceBuildNumber")), AndroidProperties.DeviceBuildNumber);
-		OutDeviceParameters.Add(FName(TEXT("SRC_UsingHoudini")), AndroidProperties.UsingHoudini ? "true" : "false");
-		OutDeviceParameters.Add(FName(TEXT("SRC_Hardware")), AndroidProperties.Hardware);
-		OutDeviceParameters.Add(FName(TEXT("SRC_Chipset")), AndroidProperties.Chipset);
-		OutDeviceParameters.Add(FName(TEXT("SRC_TotalPhysicalGB")), AndroidProperties.TotalPhysicalGB);
-		OutDeviceParameters.Add(FName(TEXT("SRC_HMDSystemName")), TEXT(""));
-		OutDeviceParameters.Add(FName(TEXT("SRC_SM5Available")), AndroidProperties.SM5Available ? "true" : "false");
-		OutDeviceParameters.Add(FName(TEXT("SRC_VKQuality")), TEXT(""));
-		OutDeviceParameters.Add(FAndroidProfileSelectorSourceProperties::SRC_ResolutionX, FString::Printf(TEXT("%d"), DeviceSpecs.ResolutionX));
-		OutDeviceParameters.Add(FAndroidProfileSelectorSourceProperties::SRC_ResolutionY, FString::Printf(TEXT("%d"), DeviceSpecs.ResolutionY));
-		OutDeviceParameters.Add(FAndroidProfileSelectorSourceProperties::SRC_InsetsLeft, FString::Printf(TEXT("%f"), DeviceSpecs.InsetsLeft));
-		OutDeviceParameters.Add(FAndroidProfileSelectorSourceProperties::SRC_InsetsTop, FString::Printf(TEXT("%f"), DeviceSpecs.InsetsTop));
-		OutDeviceParameters.Add(FAndroidProfileSelectorSourceProperties::SRC_InsetsRight, FString::Printf(TEXT("%f"), DeviceSpecs.InsetsRight));
-		OutDeviceParameters.Add(FAndroidProfileSelectorSourceProperties::SRC_InsetsBottom, FString::Printf(TEXT("%f"), DeviceSpecs.InsetsBottom));
+		OutDeviceParameters.Add(TEXT("SRC_GPUFamily"), AndroidProperties.GPUFamily);
+		OutDeviceParameters.Add(TEXT("SRC_GLVersion"), AndroidProperties.GLVersion);
+		OutDeviceParameters.Add(TEXT("SRC_VulkanAvailable"), AndroidProperties.VulkanAvailable ? "true" : "false");
+		OutDeviceParameters.Add(TEXT("SRC_VulkanVersion"), AndroidProperties.VulkanVersion);
+		OutDeviceParameters.Add(TEXT("SRC_AndroidVersion"), AndroidProperties.AndroidVersion);
+		OutDeviceParameters.Add(TEXT("SRC_DeviceMake"), AndroidProperties.DeviceMake);
+		OutDeviceParameters.Add(TEXT("SRC_DeviceModel"), AndroidProperties.DeviceModel);
+		OutDeviceParameters.Add(TEXT("SRC_DeviceBuildNumber"), AndroidProperties.DeviceBuildNumber);
+		OutDeviceParameters.Add(TEXT("SRC_UsingHoudini"), AndroidProperties.UsingHoudini ? "true" : "false");
+		OutDeviceParameters.Add(TEXT("SRC_Hardware"), AndroidProperties.Hardware);
+		OutDeviceParameters.Add(TEXT("SRC_Chipset"), AndroidProperties.Chipset);
+		OutDeviceParameters.Add(TEXT("SRC_TotalPhysicalGB"), AndroidProperties.TotalPhysicalGB);
+		OutDeviceParameters.Add(TEXT("SRC_HMDSystemName"), TEXT(""));
+		OutDeviceParameters.Add(TEXT("SRC_SM5Available"), AndroidProperties.SM5Available ? "true" : "false");
+		OutDeviceParameters.Add(TEXT("SRC_VKQuality"), TEXT(""));
+		OutDeviceParameters.Add(FString(FAndroidProfileSelectorSourceProperties::SRC_ResolutionX), FString::Printf(TEXT("%d"), DeviceSpecs.ResolutionX));
+		OutDeviceParameters.Add(FString(FAndroidProfileSelectorSourceProperties::SRC_ResolutionY), FString::Printf(TEXT("%d"), DeviceSpecs.ResolutionY));
+		OutDeviceParameters.Add(FString(FAndroidProfileSelectorSourceProperties::SRC_InsetsLeft), FString::Printf(TEXT("%f"), DeviceSpecs.InsetsLeft));
+		OutDeviceParameters.Add(FString(FAndroidProfileSelectorSourceProperties::SRC_InsetsTop), FString::Printf(TEXT("%f"), DeviceSpecs.InsetsTop));
+		OutDeviceParameters.Add(FString(FAndroidProfileSelectorSourceProperties::SRC_InsetsRight), FString::Printf(TEXT("%f"), DeviceSpecs.InsetsRight));
+		OutDeviceParameters.Add(FString(FAndroidProfileSelectorSourceProperties::SRC_InsetsBottom), FString::Printf(TEXT("%f"), DeviceSpecs.InsetsBottom));
+	}
+
+	// Initialize the ConfigRules for the Preview Device.
+	FString ConfigRulesFile = FPaths::ProjectDir() + TEXT("Build/Android/configrules.txt");
+	TArray<uint8> ConfigRulesData;
+	TMap<FString, FString> PredefinedVariables(OutDeviceParameters);
+	if (FFileHelper::LoadFileToArray(ConfigRulesData, *ConfigRulesFile))
+	{
+		const TMap<FString, FString>& PreviewConfigRules = FGenericConfigRules::ParseConfigRules(ConfigRulesData, MoveTemp(PredefinedVariables));
+		// Initialize the DeviceParams for the Preview Device.
+		for (const TPair<FString, FString>& Pair : PreviewConfigRules)
+		{
+			const FString& VariableName = Pair.Key;
+			const FString& VariableValue = Pair.Value;
+			OutDeviceParameters.Add(FString::Printf(TEXT("SRC_ConfigRuleVar[%s]"), *VariableName), *VariableValue);
+		}
 	}
 }
 
@@ -179,12 +195,24 @@ const FString FAndroidDeviceProfileSelectorModule::GetDeviceProfileName()
 
 	UE_LOG(LogAndroidDPSelector, Log, TEXT("Checking %d rules from DeviceProfile ini file."), FAndroidDeviceProfileSelector::GetNumProfiles() );
 	UE_LOG(LogAndroidDPSelector, Log, TEXT("  Default profile: %s"), *ProfileName);
-	for (const TTuple<FName,FString>& MapIt : FAndroidDeviceProfileSelector::GetSelectorProperties())
+	for (const TTuple<FString,FString>& MapIt : FAndroidDeviceProfileSelector::GetSelectorProperties())
 	{
-		UE_LOG(LogAndroidDPSelector, Log, TEXT("  %s: %s"), *MapIt.Key.ToString(), *MapIt.Value);
+		UE_LOG(LogAndroidDPSelector, Log, TEXT("  %s: %s"), *MapIt.Key, *MapIt.Value);
 	}
 
-	ProfileName = FAndroidDeviceProfileSelector::FindMatchingProfile(ProfileName);
+	// Use override from ConfigRules if set
+	const FString* ConfigProfile = FAndroidDeviceProfileSelector::GetSelectorProperties().Find(TEXT("Profile"));
+	if (ConfigProfile != nullptr)
+	{
+		ProfileName = *ConfigProfile;
+		UE_LOG(LogAndroidDPSelector, Log, TEXT("Using ConfigRules Profile: [%s]"), *ProfileName);
+	}
+	else
+	{
+		// Find a match with the DeviceProfiles matching rules
+		ProfileName = FAndroidDeviceProfileSelector::FindMatchingProfile(ProfileName);
+		UE_LOG(LogAndroidDPSelector, Log, TEXT("Selected Device Profile: [%s]"), *ProfileName);
+	}
 
 	UE_LOG(LogAndroidDPSelector, Log, TEXT("Selected Device Profile: [%s]"), *ProfileName);
 
@@ -193,7 +221,7 @@ const FString FAndroidDeviceProfileSelectorModule::GetDeviceProfileName()
 
 bool FAndroidDeviceProfileSelectorModule::GetSelectorPropertyValue(const FName& PropertyType, FString& PropertyValueOUT)
 {
-	if (const FString* Found = FAndroidDeviceProfileSelector::GetSelectorProperties().Find(PropertyType))
+	if (const FString* Found = FAndroidDeviceProfileSelector::GetSelectorProperties().Find(PropertyType.ToString()))
 	{
 		PropertyValueOUT = *Found;
 		return true;
@@ -210,15 +238,19 @@ bool FAndroidDeviceProfileSelectorModule::GetSelectorPropertyValue(const FName&
 	return false;
 }
 
-void FAndroidDeviceProfileSelectorModule::SetSelectorProperties(const TMap<FName, FString>& SelectorPropertiesIn)
+void FAndroidDeviceProfileSelectorModule::SetSelectorProperties(const TMap<FString, FString>& SelectorPropertiesIn)
 {
-	const int32 ResolutionX = FCString::Atoi(*SelectorPropertiesIn.FindChecked(FAndroidProfileSelectorSourceProperties::SRC_ResolutionX));
-	const int32 ResolutionY = FCString::Atoi(*SelectorPropertiesIn.FindChecked(FAndroidProfileSelectorSourceProperties::SRC_ResolutionY));
-	const float InsetsLeft = FCString::Atof(*SelectorPropertiesIn.FindChecked(FAndroidProfileSelectorSourceProperties::SRC_InsetsLeft));
-	const float InsetsTop = FCString::Atof(*SelectorPropertiesIn.FindChecked(FAndroidProfileSelectorSourceProperties::SRC_InsetsTop));
-	const float InsetsRight = FCString::Atof(*SelectorPropertiesIn.FindChecked(FAndroidProfileSelectorSourceProperties::SRC_InsetsRight));
-	const float InsetsBottom = FCString::Atof(*SelectorPropertiesIn.FindChecked(FAndroidProfileSelectorSourceProperties::SRC_InsetsBottom));
-
+	auto GetSelectorPropertiesValueChecked = [SelectorPropertiesIn](const FStringView& InStringView)
+		{
+			return SelectorPropertiesIn.FindByHashChecked(GetTypeHash(InStringView), InStringView);
+		};
+	const int32 ResolutionX= FCString::Atoi(*GetSelectorPropertiesValueChecked(FAndroidProfileSelectorSourceProperties::SRC_ResolutionX));
+	const int32 ResolutionY = FCString::Atoi(*GetSelectorPropertiesValueChecked(FAndroidProfileSelectorSourceProperties::SRC_ResolutionY));
+	const float InsetsLeft = FCString::Atof(*GetSelectorPropertiesValueChecked(FAndroidProfileSelectorSourceProperties::SRC_InsetsLeft));
+	const float InsetsTop = FCString::Atof(*GetSelectorPropertiesValueChecked(FAndroidProfileSelectorSourceProperties::SRC_InsetsTop));
+	const float InsetsRight = FCString::Atof(*GetSelectorPropertiesValueChecked(FAndroidProfileSelectorSourceProperties::SRC_InsetsRight));
+	const float InsetsBottom = FCString::Atof(*GetSelectorPropertiesValueChecked(FAndroidProfileSelectorSourceProperties::SRC_InsetsBottom));
+	
 	FString Orientation;
 	bool bNeedPortrait = false;
 	GConfig->GetString(TEXT("/Script/AndroidRuntimeSettings.AndroidRuntimeSettings"), TEXT("Orientation"), Orientation, GEngineIni);
diff --git a/Engine/Plugins/Runtime/AndroidDeviceProfileSelector/Source/AndroidDeviceProfileSelector/Private/AndroidDeviceProfileSelectorModule.h b/Engine/Plugins/Runtime/AndroidDeviceProfileSelector/Source/AndroidDeviceProfileSelector/Private/AndroidDeviceProfileSelectorModule.h
--- a/Engine/Plugins/Runtime/AndroidDeviceProfileSelector/Source/AndroidDeviceProfileSelector/Private/AndroidDeviceProfileSelectorModule.h
+++ b/Engine/Plugins/Runtime/AndroidDeviceProfileSelector/Source/AndroidDeviceProfileSelector/Private/AndroidDeviceProfileSelectorModule.h
@@ -19,12 +19,12 @@ public:
 #if WITH_EDITOR
 	virtual void ExportDeviceParametersToJson(FString& FolderLocation, TArray<FString>& OutExportedFiles) override;
 	virtual bool CanExportDeviceParametersToJson() override;
-	virtual void GetDeviceParametersFromJson(FString& JsonLocation, TMap<FName, FString>& OutDeviceParameters) override;
+	virtual void GetDeviceParametersFromJson(FString& JsonLocation, TMap<FString, FString>& OutDeviceParameters) override;
 	virtual bool CanGetDeviceParametersFromJson() override { return true; }
 	virtual const EShaderPlatform GetPreviewShaderPlatform() override;
 #endif
 	// Set the device parameters this selector will use.
-	virtual void SetSelectorProperties(const TMap<FName, FString>& SelectorProperties) override;
+	virtual void SetSelectorProperties(const TMap<FString, FString>& SelectorProperties) override;
 	virtual bool GetSelectorPropertyValue(const FName& PropertyType, FString& PropertyValueOUT) override;
 	//~ End IDeviceProfileSelectorModule Interface
 
diff --git a/Engine/Plugins/Runtime/AndroidDeviceProfileSelector/Source/AndroidDeviceProfileSelector/Public/AndroidDeviceProfileSelector.h b/Engine/Plugins/Runtime/AndroidDeviceProfileSelector/Source/AndroidDeviceProfileSelector/Public/AndroidDeviceProfileSelector.h
--- a/Engine/Plugins/Runtime/AndroidDeviceProfileSelector/Source/AndroidDeviceProfileSelector/Public/AndroidDeviceProfileSelector.h
+++ b/Engine/Plugins/Runtime/AndroidDeviceProfileSelector/Source/AndroidDeviceProfileSelector/Public/AndroidDeviceProfileSelector.h
@@ -9,40 +9,40 @@
 
 namespace FAndroidProfileSelectorSourceProperties
 {
-	static FName SRC_GPUFamily(TEXT("SRC_GPUFamily"));
-	static FName SRC_GLVersion(TEXT("SRC_GLVersion"));
-	static FName SRC_VulkanAvailable(TEXT("SRC_VulkanAvailable"));
-	static FName SRC_VulkanVersion(TEXT("SRC_VulkanVersion"));
-	static FName SRC_AndroidVersion(TEXT("SRC_AndroidVersion"));
-	static FName SRC_DeviceMake(TEXT("SRC_DeviceMake"));
-	static FName SRC_DeviceModel(TEXT("SRC_DeviceModel"));
-	static FName SRC_DeviceBuildNumber(TEXT("SRC_DeviceBuildNumber"));
-	static FName SRC_UsingHoudini(TEXT("SRC_UsingHoudini"));
-	static FName SRC_Hardware(TEXT("SRC_Hardware"));
-	static FName SRC_Chipset(TEXT("SRC_Chipset"));
-	static FName SRC_HMDSystemName(TEXT("SRC_HMDSystemName"));
-	static FName SRC_TotalPhysicalGB(TEXT("SRC_TotalPhysicalGB"));
-	static FName SRC_SM5Available(TEXT("SRC_SM5Available"));
-	static FName SRC_ResolutionX(TEXT("SRC_ResolutionX"));
-	static FName SRC_ResolutionY(TEXT("SRC_ResolutionY"));
-	static FName SRC_InsetsLeft(TEXT("SRC_InsetsLeft"));
-	static FName SRC_InsetsTop(TEXT("SRC_InsetsTop"));
-	static FName SRC_InsetsRight(TEXT("SRC_InsetsRight"));
-	static FName SRC_InsetsBottom(TEXT("SRC_InsetsBottom"));
-	static FName SRC_VKQuality(TEXT("SRC_VKQuality"));
+	static FStringView SRC_GPUFamily(TEXT("SRC_GPUFamily"));
+	static FStringView SRC_GLVersion(TEXT("SRC_GLVersion"));
+	static FStringView SRC_VulkanAvailable(TEXT("SRC_VulkanAvailable"));
+	static FStringView SRC_VulkanVersion(TEXT("SRC_VulkanVersion"));
+	static FStringView SRC_AndroidVersion(TEXT("SRC_AndroidVersion"));
+	static FStringView SRC_DeviceMake(TEXT("SRC_DeviceMake"));
+	static FStringView SRC_DeviceModel(TEXT("SRC_DeviceModel"));
+	static FStringView SRC_DeviceBuildNumber(TEXT("SRC_DeviceBuildNumber"));
+	static FStringView SRC_UsingHoudini(TEXT("SRC_UsingHoudini"));
+	static FStringView SRC_Hardware(TEXT("SRC_Hardware"));
+	static FStringView SRC_Chipset(TEXT("SRC_Chipset"));
+	static FStringView SRC_HMDSystemName(TEXT("SRC_HMDSystemName"));
+	static FStringView SRC_TotalPhysicalGB(TEXT("SRC_TotalPhysicalGB"));
+	static FStringView SRC_SM5Available(TEXT("SRC_SM5Available"));
+	static FStringView SRC_ResolutionX(TEXT("SRC_ResolutionX"));
+	static FStringView SRC_ResolutionY(TEXT("SRC_ResolutionY"));
+	static FStringView SRC_InsetsLeft(TEXT("SRC_InsetsLeft"));
+	static FStringView SRC_InsetsTop(TEXT("SRC_InsetsTop"));
+	static FStringView SRC_InsetsRight(TEXT("SRC_InsetsRight"));
+	static FStringView SRC_InsetsBottom(TEXT("SRC_InsetsBottom"));
+	static FStringView SRC_VKQuality(TEXT("SRC_VKQuality"));
 };
 
 class FAndroidDeviceProfileSelector
 {
 	// Container of various device properties used for device profile matching.
-	static UE_API TMap<FName, FString> SelectorProperties;
+	static UE_API TMap<FString, FString> SelectorProperties;
 	static UE_API void VerifySelectorParams();
 public:
 
 	static UE_API FString FindMatchingProfile(const FString& FallbackProfileName);
 	static UE_API int32 GetNumProfiles();
-	static void SetSelectorProperties(const TMap<FName, FString>& Params) { SelectorProperties = Params; VerifySelectorParams(); }
-	static const TMap<FName, FString>& GetSelectorProperties() { return SelectorProperties; }
+	static void SetSelectorProperties(const TMap<FString, FString>& Params) { SelectorProperties = Params; VerifySelectorParams(); }
+	static const TMap<FString, FString>& GetSelectorProperties() { return SelectorProperties; }
 };
 
 #undef UE_API
diff --git a/Engine/Plugins/Runtime/AndroidDeviceProfileSelector/Source/AndroidDeviceProfileSelectorRuntime/Private/AndroidDeviceProfileSelectorRuntimeModule.cpp b/Engine/Plugins/Runtime/AndroidDeviceProfileSelector/Source/AndroidDeviceProfileSelectorRuntime/Private/AndroidDeviceProfileSelectorRuntimeModule.cpp
--- a/Engine/Plugins/Runtime/AndroidDeviceProfileSelector/Source/AndroidDeviceProfileSelectorRuntime/Private/AndroidDeviceProfileSelectorRuntimeModule.cpp
+++ b/Engine/Plugins/Runtime/AndroidDeviceProfileSelector/Source/AndroidDeviceProfileSelectorRuntime/Private/AndroidDeviceProfileSelectorRuntimeModule.cpp
@@ -21,10 +21,10 @@ void FAndroidDeviceProfileSelectorRuntimeModule::ShutdownModule()
 }
 
 // Build the selector params for the current device.
-static const TMap<FName,FString>& GetDeviceSelectorParams()
+static const TMap<FString,FString>& GetDeviceSelectorParams()
 {
 	static bool bInitialized = false;
-	static TMap<FName, FString> AndroidParams;
+	static TMap<FString, FString> AndroidParams;
 	if(!bInitialized)
 	{
 		bInitialized = true;
@@ -64,21 +64,21 @@ static const TMap<FName,FString>& GetDeviceSelectorParams()
 			HMDRequestedProfileName = IHeadMountedDisplayModule::Get().GetDeviceSystemName();
 		}
 
-		AndroidParams.Add(FAndroidProfileSelectorSourceProperties::SRC_GPUFamily, GetParam(FAndroidMisc::GetGPUFamily(), TEXT("gpu")));
-		AndroidParams.Add(FAndroidProfileSelectorSourceProperties::SRC_GLVersion, FAndroidMisc::GetGLVersion());
-		AndroidParams.Add(FAndroidProfileSelectorSourceProperties::SRC_VulkanAvailable, FAndroidMisc::IsVulkanAvailable() ? TEXT("true") : TEXT("false"));
-		AndroidParams.Add(FAndroidProfileSelectorSourceProperties::SRC_VulkanVersion, FAndroidMisc::GetVulkanVersion());
-		AndroidParams.Add(FAndroidProfileSelectorSourceProperties::SRC_AndroidVersion, FAndroidMisc::GetAndroidVersion());
-		AndroidParams.Add(FAndroidProfileSelectorSourceProperties::SRC_DeviceMake, FAndroidMisc::GetDeviceMake());
-		AndroidParams.Add(FAndroidProfileSelectorSourceProperties::SRC_DeviceModel, FAndroidMisc::GetDeviceModel());
-		AndroidParams.Add(FAndroidProfileSelectorSourceProperties::SRC_DeviceBuildNumber, FAndroidMisc::GetDeviceBuildNumber());
-		AndroidParams.Add(FAndroidProfileSelectorSourceProperties::SRC_UsingHoudini, bUsingHoudini ? TEXT("true") : TEXT("false"));
-		AndroidParams.Add(FAndroidProfileSelectorSourceProperties::SRC_Hardware, GetParam(FString(TEXT("unknown")), TEXT("hardware")));
-		AndroidParams.Add(FAndroidProfileSelectorSourceProperties::SRC_Chipset, GetParam(FString(TEXT("unknown")), TEXT("chipset")));
-		AndroidParams.Add(FAndroidProfileSelectorSourceProperties::SRC_HMDSystemName, HMDRequestedProfileName);
-		AndroidParams.Add(FAndroidProfileSelectorSourceProperties::SRC_TotalPhysicalGB, FString::Printf(TEXT("%d"), TotalPhysicalGB));
-		AndroidParams.Add(FAndroidProfileSelectorSourceProperties::SRC_SM5Available, FAndroidMisc::IsDesktopVulkanAvailable() ? TEXT("true") : TEXT("false"));
-		AndroidParams.Add(FAndroidProfileSelectorSourceProperties::SRC_VKQuality, FAndroidMisc::GetVKQualityRecommendation());
+		AndroidParams.Add(FString(FAndroidProfileSelectorSourceProperties::SRC_GPUFamily), GetParam(FAndroidMisc::GetGPUFamily(), TEXT("gpu")));
+		AndroidParams.Add(FString(FAndroidProfileSelectorSourceProperties::SRC_GLVersion), FAndroidMisc::GetGLVersion());
+		AndroidParams.Add(FString(FAndroidProfileSelectorSourceProperties::SRC_VulkanAvailable), FAndroidMisc::IsVulkanAvailable() ? TEXT("true") : TEXT("false"));
+		AndroidParams.Add(FString(FAndroidProfileSelectorSourceProperties::SRC_VulkanVersion), FAndroidMisc::GetVulkanVersion());
+		AndroidParams.Add(FString(FAndroidProfileSelectorSourceProperties::SRC_AndroidVersion), FAndroidMisc::GetAndroidVersion());
+		AndroidParams.Add(FString(FAndroidProfileSelectorSourceProperties::SRC_DeviceMake), FAndroidMisc::GetDeviceMake());
+		AndroidParams.Add(FString(FAndroidProfileSelectorSourceProperties::SRC_DeviceModel), FAndroidMisc::GetDeviceModel());
+		AndroidParams.Add(FString(FAndroidProfileSelectorSourceProperties::SRC_DeviceBuildNumber), FAndroidMisc::GetDeviceBuildNumber());
+		AndroidParams.Add(FString(FAndroidProfileSelectorSourceProperties::SRC_UsingHoudini), bUsingHoudini ? TEXT("true") : TEXT("false"));
+		AndroidParams.Add(FString(FAndroidProfileSelectorSourceProperties::SRC_Hardware), GetParam(FString(TEXT("unknown")), TEXT("hardware")));
+		AndroidParams.Add(FString(FAndroidProfileSelectorSourceProperties::SRC_Chipset), GetParam(FString(TEXT("unknown")), TEXT("chipset")));
+		AndroidParams.Add(FString(FAndroidProfileSelectorSourceProperties::SRC_HMDSystemName), HMDRequestedProfileName);
+		AndroidParams.Add(FString(FAndroidProfileSelectorSourceProperties::SRC_TotalPhysicalGB), FString::Printf(TEXT("%d"), TotalPhysicalGB));
+		AndroidParams.Add(FString(FAndroidProfileSelectorSourceProperties::SRC_SM5Available), FAndroidMisc::IsDesktopVulkanAvailable() ? TEXT("true") : TEXT("false"));
+		AndroidParams.Add(FString(FAndroidProfileSelectorSourceProperties::SRC_VKQuality), FAndroidMisc::GetVKQualityRecommendation());
 
 #if PLATFORM_ANDROID
 		// allow ConfigRules to override cvars first
@@ -87,7 +87,7 @@ static const TMap<FName,FString>& GetDeviceSelectorParams()
 		{
 			const FString& VariableName = Pair.Key;
 			const FString& VariableValue = Pair.Value;
-			AndroidParams.Add(FName(FString::Printf(TEXT("SRC_ConfigRuleVar[%s]"), *VariableName)), *VariableValue);
+			AndroidParams.Add(FString::Printf(TEXT("SRC_ConfigRuleVar[%s]"), *VariableName), *VariableValue);
 		}
 #endif
 	}
@@ -106,7 +106,7 @@ FString const FAndroidDeviceProfileSelectorRuntimeModule::GetRuntimeDeviceProfil
 		{
 			ProfileName = FPlatformProperties::PlatformName();
 		}
-		const TMap<FName, FString>& AndroidParams = GetDeviceSelectorParams();
+		const TMap<FString, FString>& AndroidParams = GetDeviceSelectorParams();
 
 		FAndroidDeviceProfileSelector::SetSelectorProperties(AndroidParams);
 
@@ -115,11 +115,11 @@ FString const FAndroidDeviceProfileSelectorRuntimeModule::GetRuntimeDeviceProfil
 		UE_LOG(LogAndroid, Log, TEXT("  Android selector params: "));
 		for(auto& MapIt : AndroidParams)
 		{
-			UE_LOG(LogAndroid, Log, TEXT("  %s: %s"), *MapIt.Key.ToString(), *MapIt.Value);
+			UE_LOG(LogAndroid, Log, TEXT("  %s: %s"), *MapIt.Key, *MapIt.Value);
 		}
 
-		const FString& DeviceMake = AndroidParams.FindChecked(FAndroidProfileSelectorSourceProperties::SRC_DeviceMake);
-		const FString& DeviceModel = AndroidParams.FindChecked(FAndroidProfileSelectorSourceProperties::SRC_DeviceModel);
+		const FString& DeviceMake = AndroidParams.FindByHashChecked(GetTypeHash(FAndroidProfileSelectorSourceProperties::SRC_DeviceMake), FAndroidProfileSelectorSourceProperties::SRC_DeviceMake);
+		const FString& DeviceModel = AndroidParams.FindByHashChecked(GetTypeHash(FAndroidProfileSelectorSourceProperties::SRC_DeviceModel), FAndroidProfileSelectorSourceProperties::SRC_DeviceModel);
 
 		CheckForJavaSurfaceViewWorkaround(DeviceMake, DeviceModel);
 #if !UE_BUILD_SHIPPING
@@ -137,7 +137,7 @@ FString const FAndroidDeviceProfileSelectorRuntimeModule::GetRuntimeDeviceProfil
 					return ProfileName;
 				}
 				
-				if (AndroidParams.FindChecked(FAndroidProfileSelectorSourceProperties::SRC_GPUFamily).StartsWith(TEXT("Android Emulator")))
+				if (AndroidParams.FindByHashChecked(GetTypeHash(FAndroidProfileSelectorSourceProperties::SRC_GPUFamily), FAndroidProfileSelectorSourceProperties::SRC_GPUFamily).StartsWith(TEXT("Android Emulator")))
 				{
 					ProfileName = TEXT("Android_Emulator");
 					UE_LOG(LogAndroid, Log, TEXT("Forced Profile: [Android_Emulator]"));
@@ -167,7 +167,7 @@ FString const FAndroidDeviceProfileSelectorRuntimeModule::GetRuntimeDeviceProfil
 
 bool FAndroidDeviceProfileSelectorRuntimeModule::GetSelectorPropertyValue(const FName& PropertyType, FString& PropertyValueOUT)
 {
-	if (const FString* Found = GetDeviceSelectorParams().Find(PropertyType))
+	if (const FString* Found = GetDeviceSelectorParams().Find(PropertyType.ToString()))
 	{
 		PropertyValueOUT = *Found;
 		return true;
diff --git a/Engine/Plugins/Runtime/IOSDeviceProfileSelector/Source/IOSPreviewDeviceProfileSelector/Private/IOSPreviewDeviceProfileSelectorModule.cpp b/Engine/Plugins/Runtime/IOSDeviceProfileSelector/Source/IOSPreviewDeviceProfileSelector/Private/IOSPreviewDeviceProfileSelectorModule.cpp
--- a/Engine/Plugins/Runtime/IOSDeviceProfileSelector/Source/IOSPreviewDeviceProfileSelector/Private/IOSPreviewDeviceProfileSelectorModule.cpp
+++ b/Engine/Plugins/Runtime/IOSDeviceProfileSelector/Source/IOSPreviewDeviceProfileSelector/Private/IOSPreviewDeviceProfileSelectorModule.cpp
@@ -140,7 +140,7 @@ void FIOSPreviewDeviceProfileSelectorModule::ExportDeviceParametersToJson(FStrin
 #endif
 }
 
-void FIOSPreviewDeviceProfileSelectorModule::GetDeviceParametersFromJson(FString& JsonLocation, TMap<FName, FString>& OutDeviceParameters)
+void FIOSPreviewDeviceProfileSelectorModule::GetDeviceParametersFromJson(FString& JsonLocation, TMap<FString, FString>& OutDeviceParameters)
 {
 	TSharedPtr<FJsonObject> JsonRootObject;
 	FString Json;
@@ -155,7 +155,7 @@ void FIOSPreviewDeviceProfileSelectorModule::GetDeviceParametersFromJson(FString
 	if (JsonRootObject.IsValid() && FJsonObjectConverter::JsonAttributesToUStruct(JsonRootObject->Values, FPIEPreviewDeviceSpecifications::StaticStruct(), &DeviceSpecs, 0, 0))
 	{
 		FPIEIOSDeviceProperties& IOSProperties = DeviceSpecs.IOSProperties;
-		OutDeviceParameters.Add(FName(TEXT("SRC_DeviceModel")), IOSProperties.DeviceModel);
+		OutDeviceParameters.Add((TEXT("SRC_DeviceModel")), IOSProperties.DeviceModel);
 	}
 }
 
@@ -178,14 +178,14 @@ const EShaderPlatform FIOSPreviewDeviceProfileSelectorModule::GetPreviewShaderPl
 	return PreviewPlatform;
 }
 
-void FIOSPreviewDeviceProfileSelectorModule::SetSelectorProperties(const TMap<FName, FString>& InSelectorProperties)
+void FIOSPreviewDeviceProfileSelectorModule::SetSelectorProperties(const TMap<FString, FString>& InSelectorProperties)
 {
 	SelectorProperties = InSelectorProperties;
 }
 
 bool FIOSPreviewDeviceProfileSelectorModule::GetSelectorPropertyValue(const FName& PropertyType, FString& PropertyValueOUT)
 {
-	if (const FString* Found = SelectorProperties.Find(PropertyType))
+	if (const FString* Found = SelectorProperties.Find(PropertyType.ToString()))
 	{
 		PropertyValueOUT = *Found;
 		return true;
diff --git a/Engine/Plugins/Runtime/IOSDeviceProfileSelector/Source/IOSPreviewDeviceProfileSelector/Private/IOSPreviewDeviceProfileSelectorModule.h b/Engine/Plugins/Runtime/IOSDeviceProfileSelector/Source/IOSPreviewDeviceProfileSelector/Private/IOSPreviewDeviceProfileSelectorModule.h
--- a/Engine/Plugins/Runtime/IOSDeviceProfileSelector/Source/IOSPreviewDeviceProfileSelector/Private/IOSPreviewDeviceProfileSelectorModule.h
+++ b/Engine/Plugins/Runtime/IOSDeviceProfileSelector/Source/IOSPreviewDeviceProfileSelector/Private/IOSPreviewDeviceProfileSelectorModule.h
@@ -19,10 +19,10 @@ public:
 #if WITH_EDITOR
 	virtual bool CanExportDeviceParametersToJson() override;
 	virtual void ExportDeviceParametersToJson(FString& FolderLocation, TArray<FString>& OutExportedFiles) override;
-	virtual void GetDeviceParametersFromJson(FString& JsonLocation, TMap<FName, FString>& OutDeviceParameters) override;
+	virtual void GetDeviceParametersFromJson(FString& JsonLocation, TMap<FString, FString>& OutDeviceParameters) override;
 	virtual bool CanGetDeviceParametersFromJson() override { return true; }
 	virtual const EShaderPlatform GetPreviewShaderPlatform() override;
-	virtual void SetSelectorProperties(const TMap<FName, FString>& InSelectorProperties) override;
+	virtual void SetSelectorProperties(const TMap<FString, FString>& InSelectorProperties) override;
 	virtual bool GetSelectorPropertyValue(const FName& PropertyType, FString& PropertyValueOUT) override;
 #endif
 	//~ End IDeviceProfileSelectorModule Interface
@@ -41,5 +41,5 @@ public:
 	{
 	}
 private:
-	TMap<FName, FString> SelectorProperties;
+	TMap<FString, FString> SelectorProperties;
 };
diff --git a/Engine/Source/Editor/DeviceProfileServices/Private/CheckAndroidDeviceProfileCommandlet.cpp b/Engine/Source/Editor/DeviceProfileServices/Private/CheckAndroidDeviceProfileCommandlet.cpp
--- a/Engine/Source/Editor/DeviceProfileServices/Private/CheckAndroidDeviceProfileCommandlet.cpp
+++ b/Engine/Source/Editor/DeviceProfileServices/Private/CheckAndroidDeviceProfileCommandlet.cpp
@@ -114,7 +114,7 @@ int32 UCheckAndroidDeviceProfileCommandlet::Main(const FString& RawCommandLine)
 	{
 		do
 		{
-			TMap<FName, FString> DeviceParameters;
+			TMap<FString, FString> DeviceParameters;
 			FString DeviceName;
 			FString OutputFilename;
 			if (DeviceSpecificationFileNames.Num())
@@ -133,32 +133,32 @@ int32 UCheckAndroidDeviceProfileCommandlet::Main(const FString& RawCommandLine)
 			else
 			{
 				DeviceName = TEXT("[no name]");
-				DeviceParameters.Add(FName(TEXT("SRC_GPUFamily")), Params.FindRef(TEXT("GPUFamily")));
-				DeviceParameters.Add(FName(TEXT("SRC_GLVersion")), Params.FindRef(TEXT("GLVersion")));
-				DeviceParameters.Add(FName(TEXT("SRC_VulkanAvailable")), Params.FindRef(TEXT("VulkanAvailable")));
-				DeviceParameters.Add(FName(TEXT("SRC_VulkanVersion")), Params.FindRef(TEXT("VulkanVersion")));
-				DeviceParameters.Add(FName(TEXT("SRC_AndroidVersion")), Params.FindRef(TEXT("AndroidVersion")));
-				DeviceParameters.Add(FName(TEXT("SRC_DeviceMake")),
+				DeviceParameters.Add(TEXT("SRC_GPUFamily"), Params.FindRef(TEXT("GPUFamily")));
+				DeviceParameters.Add(TEXT("SRC_GLVersion"), Params.FindRef(TEXT("GLVersion")));
+				DeviceParameters.Add(TEXT("SRC_VulkanAvailable"), Params.FindRef(TEXT("VulkanAvailable")));
+				DeviceParameters.Add(TEXT("SRC_VulkanVersion"), Params.FindRef(TEXT("VulkanVersion")));
+				DeviceParameters.Add(TEXT("SRC_AndroidVersion"), Params.FindRef(TEXT("AndroidVersion")));
+				DeviceParameters.Add(TEXT("SRC_DeviceMake"),
 					Tokens.Num() == 2 ? Tokens[0] :
 					Params.FindRef(TEXT("DeviceMake")));
-				DeviceParameters.Add(FName(TEXT("SRC_DeviceModel")),
+				DeviceParameters.Add(TEXT("SRC_DeviceModel"),
 					Tokens.Num() == 1 ? Tokens[0] :
 					Tokens.Num() == 2 ? Tokens[1] :
 					Params.FindRef(TEXT("DeviceModel")));
-				DeviceParameters.Add(FName(TEXT("SRC_DeviceBuildNumber")), Params.FindRef(TEXT("DeviceBuildNumber")));
-				DeviceParameters.Add(FName(TEXT("SRC_UsingHoudini")), Params.FindRef(TEXT("UsingHoudini")));
-				DeviceParameters.Add(FName(TEXT("SRC_Hardware")), Params.FindRef(TEXT("Hardware")));
-				DeviceParameters.Add(FName(TEXT("SRC_Chipset")), Params.FindRef(TEXT("Chipset")));
-				DeviceParameters.Add(FName(TEXT("SRC_TotalPhysicalGB")), Params.FindRef(TEXT("TotalPhysicalGB")));
-				DeviceParameters.Add(FName(TEXT("SRC_HMDSystemName")), TEXT(""));
-				DeviceParameters.Add(FName(TEXT("SRC_SM5Available")), Params.FindRef(TEXT("SM5Available")));
-				DeviceParameters.Add(FName(TEXT("SRC_ResolutionX")), Params.FindRef(TEXT("ResolutionX")));
-				DeviceParameters.Add(FName(TEXT("SRC_ResolutionY")), Params.FindRef(TEXT("ResolutionY")));
-				DeviceParameters.Add(FName(TEXT("SRC_InsetsLeft")), Params.FindRef(TEXT("InsetsLeft")));
-				DeviceParameters.Add(FName(TEXT("SRC_InsetsTop")), Params.FindRef(TEXT("InsetsTop")));
-				DeviceParameters.Add(FName(TEXT("SRC_InsetsRight")), Params.FindRef(TEXT("InsetsRight")));
-				DeviceParameters.Add(FName(TEXT("SRC_InsetsBottom")), Params.FindRef(TEXT("InsetsBottom")));
-				DeviceParameters.Add(FName(TEXT("SRC_VKQuality")), Params.FindRef(TEXT("VKQuality")));
+				DeviceParameters.Add(TEXT("SRC_DeviceBuildNumber"), Params.FindRef(TEXT("DeviceBuildNumber")));
+				DeviceParameters.Add(TEXT("SRC_UsingHoudini"), Params.FindRef(TEXT("UsingHoudini")));
+				DeviceParameters.Add(TEXT("SRC_Hardware"), Params.FindRef(TEXT("Hardware")));
+				DeviceParameters.Add(TEXT("SRC_Chipset"), Params.FindRef(TEXT("Chipset")));
+				DeviceParameters.Add(TEXT("SRC_TotalPhysicalGB"), Params.FindRef(TEXT("TotalPhysicalGB")));
+				DeviceParameters.Add(TEXT("SRC_HMDSystemName"), TEXT(""));
+				DeviceParameters.Add(TEXT("SRC_SM5Available"), Params.FindRef(TEXT("SM5Available")));
+				DeviceParameters.Add(TEXT("SRC_ResolutionX"), Params.FindRef(TEXT("ResolutionX")));
+				DeviceParameters.Add(TEXT("SRC_ResolutionY"), Params.FindRef(TEXT("ResolutionY")));
+				DeviceParameters.Add(TEXT("SRC_InsetsLeft"), Params.FindRef(TEXT("InsetsLeft")));
+				DeviceParameters.Add(TEXT("SRC_InsetsTop"), Params.FindRef(TEXT("InsetsTop")));
+				DeviceParameters.Add(TEXT("SRC_InsetsRight"), Params.FindRef(TEXT("InsetsRight")));
+				DeviceParameters.Add(TEXT("SRC_InsetsBottom"), Params.FindRef(TEXT("InsetsBottom")));
+				DeviceParameters.Add(TEXT("SRC_VKQuality"), Params.FindRef(TEXT("VKQuality")));
 			}
 			AndroidDeviceProfileSelector->SetSelectorProperties(DeviceParameters);
 			FString ProfileName = ForcedDeviceProfileName.IsEmpty() ? AndroidDeviceProfileSelector->GetDeviceProfileName() : ForcedDeviceProfileName;
@@ -193,29 +193,29 @@ int32 UCheckAndroidDeviceProfileCommandlet::Main(const FString& RawCommandLine)
 
 			if (!OutputFilename.IsEmpty())
 			{
-				FString output;
+				FStringBuilderBase Output;
 
 				if (!DeviceName.IsEmpty())
 				{
-					output += FString::Printf(TEXT("Device name : %s\n\n"), *DeviceName);
+					Output.Appendf(TEXT("Device name : %s\n\n"), *DeviceName);
 				}
 				if (!Profile)
 				{
-					output += FString::Printf(TEXT("Failed to find device profile : %s\n\n"), *ProfileName);
+					Output.Appendf(TEXT("Failed to find device profile : %s\n\n"), *ProfileName);
 				}
 
 				OutputFilename = DeviceName + TEXT("_") + DeviceParameters.FindChecked(TEXT("SRC_GPUFamily")) + TEXT("_") + DeviceParameters.FindChecked(TEXT("SRC_TotalPhysicalGB")) + TEXT("GB.txt");
-				output += FString::Printf(TEXT("AndroidDeviceProfileSelector Params:\n"));
-				TArray<FName> SortedKeys;
+				Output.Appendf(TEXT("AndroidDeviceProfileSelector Params:\n"));
+				TArray<FString> SortedKeys;
 				DeviceParameters.GenerateKeyArray(SortedKeys);
 				//SortedKeys.Sort();
-				for (FName& ParamKey : SortedKeys)
+				for (const FString& ParamKey : SortedKeys)
 				{
-					output += FString::Printf(TEXT("[%s]=%s\n"), *ParamKey.ToString(), *DeviceParameters.FindChecked(ParamKey));
+					Output.Appendf(TEXT("[%s]=%s\n"), *ParamKey, *DeviceParameters.FindChecked(ParamKey));
 				}
-				output += FString::Printf(TEXT("\nDevice Mem Bucket: %s\n"), LexToString(MemorySizeBucket));
+				Output.Appendf(TEXT("\nDevice Mem Bucket: %s\n"), LexToString(MemorySizeBucket));
 
-				output += FString::Printf(TEXT("\nSelected fragments :\n%s"), SelectedFragments.Num() == 0 ? TEXT("-none-\n") : TEXT(""));
+				Output.Appendf(TEXT("\nSelected fragments :\n%s"), SelectedFragments.Num() == 0 ? TEXT("-none-\n") : TEXT(""));
 				for (FSelectedFragmentProperties& Fragment : SelectedFragments)
 				{
 					FString FragTag;
@@ -223,16 +223,16 @@ int32 UCheckAndroidDeviceProfileCommandlet::Main(const FString& RawCommandLine)
 					{
 						FragTag = TEXT("[") + Fragment.Tag.ToString() + TEXT("]");
 					}
-					output += FString::Printf(TEXT("%s%s : Enabled=%s\n"), *FragTag, *Fragment.Fragment, Fragment.bEnabled ? TEXT("true") : TEXT("false"));
+					Output.Appendf(TEXT("%s%s : Enabled=%s\n"), *FragTag, *Fragment.Fragment, Fragment.bEnabled ? TEXT("true") : TEXT("false"));
 				}
 
-				output += FString::Printf(TEXT("\n\nSelected device Profile: %s\nSelected device profile description: %s\n"), *ProfileName, *ProfileDescription);
+				Output.Appendf(TEXT("\n\nSelected device Profile: %s\nSelected device profile description: %s\n"), *ProfileName, *ProfileDescription);
 
 				// may not want this:
 				const bool bSortCVars = true;
 				if (bSortCVars)
 				{
-					output += FString::Printf(TEXT("Sorted "));
+					Output.Appendf(TEXT("Sorted "));
 					CVars.Sort();
 				}
 				FString CVarsOnly = FString::Printf(TEXT("CVars:\n"));
@@ -241,8 +241,8 @@ int32 UCheckAndroidDeviceProfileCommandlet::Main(const FString& RawCommandLine)
 				{
 					CVarsOnly += FString::Printf(TEXT("%s\n"), *CVar);
 				}
-				output += CVarsOnly;
-				FFileHelper::SaveStringToFile(output, *(OutputFolder / OutputFilename));
+				Output.Append(CVarsOnly);
+				FFileHelper::SaveStringToFile(Output.ToString(), *(OutputFolder / OutputFilename));
 
 				FFileHelper::SaveStringToFile(CVarsOnly, *(OutputFolder / TEXT("CVarsOnly") / OutputFilename));
 			}
diff --git a/Engine/Source/Editor/LevelEditor/Private/LevelEditorActions.cpp b/Engine/Source/Editor/LevelEditor/Private/LevelEditorActions.cpp
--- a/Engine/Source/Editor/LevelEditor/Private/LevelEditorActions.cpp
+++ b/Engine/Source/Editor/LevelEditor/Private/LevelEditorActions.cpp
@@ -589,7 +589,7 @@ void FLevelEditorActionCallbacks::PreviewJson_Clicked(FName PlatformName, FStrin
 			{
 				FString DevileProfileName;
 
-				TMap<FName, FString> DeviceParameters;
+				TMap<FString, FString> DeviceParameters;
 				DPSelectorModule->GetDeviceParametersFromJson(JsonFile, DeviceParameters);
 				DPSelectorModule->SetSelectorProperties(DeviceParameters);
 				float ConstrainedAspectRatio = DPSelectorModule->GetConstrainedAspectRatio();
diff --git a/Engine/Source/Editor/PIEPreviewDeviceProfileSelector/Private/PIEPreviewDevice.cpp b/Engine/Source/Editor/PIEPreviewDeviceProfileSelector/Private/PIEPreviewDevice.cpp
--- a/Engine/Source/Editor/PIEPreviewDeviceProfileSelector/Private/PIEPreviewDevice.cpp
+++ b/Engine/Source/Editor/PIEPreviewDeviceProfileSelector/Private/PIEPreviewDevice.cpp
@@ -457,22 +457,22 @@ FString FPIEPreviewDevice::GetProfile() const
 			{
 				FPIEAndroidDeviceProperties& AndroidProperties = DeviceSpecs->AndroidProperties;
 
-				TMap<FName, FString> DeviceParameters;
-				DeviceParameters.Add(FName(TEXT("SRC_GPUFamily")), AndroidProperties.GPUFamily);
-				DeviceParameters.Add(FName(TEXT("SRC_GLVersion")), AndroidProperties.GLVersion);
-				DeviceParameters.Add(FName(TEXT("SRC_VulkanAvailable")), AndroidProperties.VulkanAvailable ? "true" : "false");
-				DeviceParameters.Add(FName(TEXT("SRC_VulkanVersion")), AndroidProperties.VulkanVersion);
-				DeviceParameters.Add(FName(TEXT("SRC_AndroidVersion")), AndroidProperties.AndroidVersion);
-				DeviceParameters.Add(FName(TEXT("SRC_DeviceMake")), AndroidProperties.DeviceMake);
-				DeviceParameters.Add(FName(TEXT("SRC_DeviceModel")), AndroidProperties.DeviceModel);
-				DeviceParameters.Add(FName(TEXT("SRC_DeviceBuildNumber")), AndroidProperties.DeviceBuildNumber);
-				DeviceParameters.Add(FName(TEXT("SRC_UsingHoudini")), AndroidProperties.UsingHoudini ? "true" : "false");
-				DeviceParameters.Add(FName(TEXT("SRC_Hardware")), AndroidProperties.Hardware);
-				DeviceParameters.Add(FName(TEXT("SRC_Chipset")), AndroidProperties.Chipset);
-				DeviceParameters.Add(FName(TEXT("SRC_HMDSystemName")), AndroidProperties.HMDSystemName);
-				DeviceParameters.Add(FName(TEXT("SRC_TotalPhysicalGB")), AndroidProperties.TotalPhysicalGB);
-				DeviceParameters.Add(FName(TEXT("SRC_SM5Available")), AndroidProperties.SM5Available ? "true" : "false");
-				DeviceParameters.Add(FName(TEXT("SRC_VKQuality")), TEXT(""));
+				TMap<FString, FString> DeviceParameters;
+				DeviceParameters.Add(TEXT("SRC_GPUFamily"), AndroidProperties.GPUFamily);
+				DeviceParameters.Add(TEXT("SRC_GLVersion"), AndroidProperties.GLVersion);
+				DeviceParameters.Add(TEXT("SRC_VulkanAvailable"), AndroidProperties.VulkanAvailable ? "true" : "false");
+				DeviceParameters.Add(TEXT("SRC_VulkanVersion"), AndroidProperties.VulkanVersion);
+				DeviceParameters.Add(TEXT("SRC_AndroidVersion"), AndroidProperties.AndroidVersion);
+				DeviceParameters.Add(TEXT("SRC_DeviceMake"), AndroidProperties.DeviceMake);
+				DeviceParameters.Add(TEXT("SRC_DeviceModel"), AndroidProperties.DeviceModel);
+				DeviceParameters.Add(TEXT("SRC_DeviceBuildNumber"), AndroidProperties.DeviceBuildNumber);
+				DeviceParameters.Add(TEXT("SRC_UsingHoudini"), AndroidProperties.UsingHoudini ? "true" : "false");
+				DeviceParameters.Add(TEXT("SRC_Hardware"), AndroidProperties.Hardware);
+				DeviceParameters.Add(TEXT("SRC_Chipset"), AndroidProperties.Chipset);
+				DeviceParameters.Add(TEXT("SRC_HMDSystemName"), AndroidProperties.HMDSystemName);
+				DeviceParameters.Add(TEXT("SRC_TotalPhysicalGB"), AndroidProperties.TotalPhysicalGB);
+				DeviceParameters.Add(TEXT("SRC_SM5Available"), AndroidProperties.SM5Available ? "true" : "false");
+				DeviceParameters.Add(TEXT("SRC_VKQuality"), TEXT(""));
 
 				AndroidDeviceProfileSelector->SetSelectorProperties(DeviceParameters);
 				FString PIEProfileName = AndroidDeviceProfileSelector->GetDeviceProfileName();
diff --git a/Engine/Source/Runtime/Engine/Public/IDeviceProfileSelectorModule.h b/Engine/Source/Runtime/Engine/Public/IDeviceProfileSelectorModule.h
--- a/Engine/Source/Runtime/Engine/Public/IDeviceProfileSelectorModule.h
+++ b/Engine/Source/Runtime/Engine/Public/IDeviceProfileSelectorModule.h
@@ -52,7 +52,7 @@ public:
 	* @Param JsonLocation Json location
 	* @Param OutDeviceParameters the device parameters read from Json
 	 */
-	virtual void GetDeviceParametersFromJson(FString& JsonLocation, TMap<FName, FString>& OutDeviceParameters) {}
+	virtual void GetDeviceParametersFromJson(FString& JsonLocation, TMap<FString, FString>& OutDeviceParameters) {}
 	
 	/**
 	 * Get the ConstrainedAspectRatio of the Device
@@ -89,7 +89,7 @@ public:
 	* Set or override the selector specific properties.
 	* @param SelectorProperties	A map of parameters to be used by device profile matching logic.
 	*/
-	virtual void SetSelectorProperties(const TMap<FName, FString>& SelectorProperties) { }
+	virtual void SetSelectorProperties(const TMap<FString, FString>& SelectorProperties) { }
 
 	/*
 	* Find a custom profile selector property value.
