[AndroidDeviceProfileSelectorRuntime]
* Updated code calling private reflection functions to match updated generated code

diff --git a/Engine/Plugins/Runtime/AndroidDeviceProfileSelector/Source/AndroidDeviceProfileSelectorRuntime/Private/AndroidDeviceProfileSelectorRuntimeModule.cpp b/Engine/Plugins/Runtime/AndroidDeviceProfileSelector/Source/AndroidDeviceProfileSelectorRuntime/Private/AndroidDeviceProfileSelectorRuntimeModule.cpp
--- a/Engine/Plugins/Runtime/AndroidDeviceProfileSelector/Source/AndroidDeviceProfileSelectorRuntime/Private/AndroidDeviceProfileSelectorRuntimeModule.cpp
+++ b/Engine/Plugins/Runtime/AndroidDeviceProfileSelector/Source/AndroidDeviceProfileSelectorRuntime/Private/AndroidDeviceProfileSelectorRuntimeModule.cpp
@@ -189,8 +189,8 @@ void FAndroidDeviceProfileSelectorRuntimeModule::CheckForJavaSurfaceViewWorkarou
 #if USE_ANDROID_JNI
 	// We need to initialize the class early as device profiles need to be evaluated before ProcessNewlyLoadedUObjects can be called.
 #if !UE_WITH_CONSTINIT_UOBJECT
-	extern UClass* Z_Construct_UClass_UAndroidJavaSurfaceViewDevices();
-	Z_Construct_UClass_UAndroidJavaSurfaceViewDevices();
+	extern UClass* Z_Construct_UClass_UAndroidJavaSurfaceViewDevices(ETypeConstructPhase);
+	Z_Construct_UClass_UAndroidJavaSurfaceViewDevices(ETypeConstructPhase::Outer);
 #endif // !UE_WITH_CONSTINIT_UOBJECT
 
 	const UAndroidJavaSurfaceViewDevices *const SurfaceViewDevices = Cast<UAndroidJavaSurfaceViewDevices>(UAndroidJavaSurfaceViewDevices::StaticClass()->GetDefaultObject());
